First pass for mirroring support in UAF

- Added both a base and additive trait for mirroring
- Added helper methods for mirroring an UAF Pose
- Added graph template for mirror node

[RN] minor

diff --git a/Engine/Plugins/Experimental/UAF/UAFMirroring/Source/UAFMirroring/Internal/MirroringTraitData.h b/Engine/Plugins/Experimental/UAF/UAFMirroring/Source/UAFMirroring/Internal/MirroringTraitData.h
new file mode 100644
--- /dev/null
+++ b/Engine/Plugins/Experimental/UAF/UAFMirroring/Source/UAFMirroring/Internal/MirroringTraitData.h
@@ -0,0 +1,91 @@
+﻿// Copyright Epic Games, Inc. All Rights Reserved.
+
+#pragma once
+
+#include "TraitCore/TraitBinding.h"
+#include "TraitCore/TraitSharedData.h"
+#include "Animation/MirrorDataTable.h"
+
+#include "MirroringTraitData.generated.h"
+
+/**
+ * Mirroring Settings
+ *
+ * This struct holds all the information need to mirror keyframe data (i.e. Bones, Curves, Attributes)
+ *
+ * @see UMirrorDataTable, FMirroringTrait, FMirroringAdditiveTrait
+ */
+USTRUCT(BlueprintType, meta = (DisplayName = "Mirroring Trait Settings"))
+struct FUAFMirroringTraitSettings
+{
+	GENERATED_BODY()
+
+	// Whether to perform mirror pass
+	UPROPERTY(EditAnywhere, Category = Settings, meta=(DisplayName = "Mirror"))
+	bool bShouldMirror = true;
+
+	// Whether to mirror bone transforms
+	UPROPERTY(EditAnywhere, Category = Channels, meta=(DisplayName = "Bones", EditCondition="bShouldMirror", EditConditionHides))
+	bool bShouldMirrorBones = true;
+
+	// Whether to mirror animation curves
+	UPROPERTY(EditAnywhere, Category = Channels, meta=(DisplayName = "Curves", EditCondition="bShouldMirror", EditConditionHides))
+	bool bShouldMirrorCurves = true;
+
+	// Whether to mirror attributes
+	UPROPERTY(EditAnywhere, Category = Channels, meta=(DisplayName = "Attributes", EditCondition="bShouldMirror", EditConditionHides))
+	bool bShouldMirrorAttributes = true;
+
+	// Data table to map bones to their mirrored counterpart
+	UPROPERTY(EditAnywhere, Category = Settings, meta = (ExportAsReference="true"))
+	TObjectPtr<const UMirrorDataTable> MirrorDataTable = nullptr;
+
+	// @todo: We need transitions settings (to properly handle DeadBlending).
+};
+
+/** A trait that can mirror an input's keyframe data. */
+USTRUCT(meta = (DisplayName = "Mirroring", ShowTooltip = true))
+struct FMirroringTraitSharedData : public FAnimNextTraitSharedData
+{
+	GENERATED_BODY()
+	
+	UPROPERTY()
+	FAnimNextTraitHandle Input;
+
+	UPROPERTY(EditAnywhere, Category = "Settings", meta=(ShowOnlyInnerProperties))
+	FUAFMirroringTraitSettings Settings;
+	
+	// Latent pin support boilerplate
+	#define TRAIT_LATENT_PROPERTIES_ENUMERATOR(GeneratorMacro) \
+	GeneratorMacro(Input) \
+	GeneratorMacro(Settings) \
+
+	GENERATE_TRAIT_LATENT_PROPERTIES(FMirroringTraitSharedData, TRAIT_LATENT_PROPERTIES_ENUMERATOR)
+	#undef TRAIT_LATENT_PROPERTIES_ENUMERATOR
+};
+
+/** Same behaviour as FMirroringTrait, but as an additive (i.e. it only mirrors the super-trait’s output). */
+USTRUCT(meta = (DisplayName = "Mirroring", ShowTooltip = true))
+struct FMirroringAdditiveTraitSharedData : public FAnimNextTraitSharedData
+{
+	GENERATED_BODY()
+
+	UPROPERTY(EditAnywhere, Category = "Settings", meta=(ShowOnlyInnerProperties))
+	FUAFMirroringTraitSettings Settings;
+
+	// Latent pin support boilerplate
+	#define TRAIT_LATENT_PROPERTIES_ENUMERATOR(GeneratorMacro) \
+	GeneratorMacro(Settings) \
+
+	GENERATE_TRAIT_LATENT_PROPERTIES(FMirroringAdditiveTraitSharedData, TRAIT_LATENT_PROPERTIES_ENUMERATOR)
+	#undef TRAIT_LATENT_PROPERTIES_ENUMERATOR
+};
+
+namespace UE::UAF
+{
+	// Namespaced alias
+	using FMirroringTraitSettings = FUAFMirroringTraitSettings;
+	using FMirroringTraitData = FMirroringTraitSharedData;
+	using FMirroringAdditiveTraitData = FMirroringAdditiveTraitSharedData;
+	
+}
\ No newline at end of file
diff --git a/Engine/Plugins/Experimental/UAF/UAFMirroring/Source/UAFMirroring/Private/Mirroring.cpp b/Engine/Plugins/Experimental/UAF/UAFMirroring/Source/UAFMirroring/Private/Mirroring.cpp
new file mode 100644
--- /dev/null
+++ b/Engine/Plugins/Experimental/UAF/UAFMirroring/Source/UAFMirroring/Private/Mirroring.cpp
@@ -0,0 +1,327 @@
+// Copyright Epic Games, Inc. All Rights Reserved.
+
+#include "Mirroring.h"
+
+#include "EvaluationVM/EvaluationVM.h"
+#include "TransformArrayOperations.h"
+#include "Animation/MirrorDataTable.h"
+#include "LODPose.h"
+#include "AnimationRuntime.h"
+#include "Engine/DataTable.h"
+#include "Logging/StructuredLog.h"
+#include "Engine/SkeletalMesh.h"
+
+namespace UE::UAF
+{
+	using FMemStackSetAllocator = TSetAllocator<TSparseArrayAllocator<FAnimStackAllocator, FAnimStackAllocator>, FAnimStackAllocator>;
+
+	int32 GetNumOfBonesForMirrorData(const FReferencePose& InReferencePose)
+	{
+		return InReferencePose.GetNumBonesForLOD(0);
+	}
+	
+	void BuildReferencePoseMirrorData(
+		const FReferencePose & InReferencePose,
+		EAxis::Type InMirrorAxis,
+		TConstArrayView<FBoneIndexType> InMeshBoneIndexToMirroredMeshBoneIndexMap,
+		TArrayView<FQuat> OutRefPoseMeshSpaceRotations,
+		TArrayView<FQuat> OutRefPoseMeshSpaceRotationCorrections)
+	{
+		const int32 RefPoseBoneNum = InReferencePose.ReferenceLocalTransforms.Num();
+		
+		const TArrayView<const FBoneIndexType>& MeshBoneIndexToParentMeshBoneIndexMap = InReferencePose.MeshBoneIndexToParentMeshBoneIndexMap;
+
+		// All the reference arrays are sized of LOD0 since we can always truncate the arrays if we need higher LOD levels. 
+		checkf(OutRefPoseMeshSpaceRotations.Num() == RefPoseBoneNum, TEXT("Buffer mismatch: %d:%d"), OutRefPoseMeshSpaceRotations.Num() , RefPoseBoneNum);
+		checkf(OutRefPoseMeshSpaceRotationCorrections.Num() == RefPoseBoneNum, TEXT("Buffer mismatch: %d:%d"), OutRefPoseMeshSpaceRotationCorrections.Num(), RefPoseBoneNum);
+		checkf(InMeshBoneIndexToMirroredMeshBoneIndexMap.Num() == RefPoseBoneNum, TEXT("Buffer mismatch: %d:%d"), InMeshBoneIndexToMirroredMeshBoneIndexMap.Num(), RefPoseBoneNum)
+		checkf(MeshBoneIndexToParentMeshBoneIndexMap.Num() == RefPoseBoneNum, TEXT("Buffer mismatch: %d:%d"), MeshBoneIndexToParentMeshBoneIndexMap.Num(), RefPoseBoneNum);
+		
+		// Copy local ref pose rotations.
+		FMemory::Memcpy(OutRefPoseMeshSpaceRotations.GetData(), InReferencePose.ReferenceLocalTransforms.Rotations.GetData(), RefPoseBoneNum * sizeof(FQuat));
+
+		// Convert the local ref pose rotations to component space.
+		// @note: We assume the root bone is at index 0.
+		for (int32 RefBoneIndex = 1; RefBoneIndex < RefPoseBoneNum; ++RefBoneIndex)
+		{
+			const FBoneIndexType ParentBoneIndex = MeshBoneIndexToParentMeshBoneIndexMap[RefBoneIndex];
+			
+			const FQuat ParentLocalSpaceRotation = OutRefPoseMeshSpaceRotations[ParentBoneIndex];
+			const FQuat RefBoneLocalSpaceRotation = OutRefPoseMeshSpaceRotations[RefBoneIndex];
+			
+			FQuat MeshSpaceRotation =  ParentLocalSpaceRotation * RefBoneLocalSpaceRotation;
+			MeshSpaceRotation.Normalize();
+			
+			OutRefPoseMeshSpaceRotations[RefBoneIndex] = MeshSpaceRotation;
+		}
+
+		// Now we can precompute the corrective rotation to align result with target space's rest orientation.
+		OutRefPoseMeshSpaceRotationCorrections[0] = FQuat::Identity;
+		for (FBoneIndexType RefBoneIndex = 1; RefBoneIndex < RefPoseBoneNum; ++RefBoneIndex)
+		{
+			const FBoneIndexType SourceBoneIndex = RefBoneIndex;
+			const FBoneIndexType TargetBoneIndex = InMeshBoneIndexToMirroredMeshBoneIndexMap[SourceBoneIndex];
+
+			// Not mapped, so skip.
+			if (TargetBoneIndex == FBoneIndexType(INDEX_NONE))
+			{
+				continue;
+			}
+			
+			OutRefPoseMeshSpaceRotationCorrections[RefBoneIndex] = FAnimationRuntime::MirrorQuat(OutRefPoseMeshSpaceRotations[SourceBoneIndex], InMirrorAxis).Inverse() * OutRefPoseMeshSpaceRotations[TargetBoneIndex];
+			OutRefPoseMeshSpaceRotationCorrections[RefBoneIndex].Normalize();
+		}
+	}
+	
+	// Get a map from source bone index to mirrored bone index.
+	void BuildMeshBoneIndexMirrorMap(
+		const FReferencePose & InReferencePose,
+		const UMirrorDataTable& InMirrorDataTable,
+		TArrayView<FBoneIndexType> OutMeshBoneIndexToMirroredMeshBoneIndexMap)
+	{
+		// Add this in case this code is called outside update/eval pass in UAF.
+		FMemMark Mark(FMemStack::Get());
+		
+		const int32 RefPoseBoneNum = InReferencePose.ReferenceLocalTransforms.Num();
+		
+		// All the reference arrays are sized of LOD0 since we can always truncate the arrays if we need higher LOD levels. 
+		checkf(RefPoseBoneNum == OutMeshBoneIndexToMirroredMeshBoneIndexMap.Num(), TEXT("Buffer mismatch: %d:%d"), RefPoseBoneNum, OutMeshBoneIndexToMirroredMeshBoneIndexMap.Num());
+		
+		// Reset the mirror table to defaults (no mirroring).
+		FMemory::Memset(OutMeshBoneIndexToMirroredMeshBoneIndexMap.GetData(), INDEX_NONE, RefPoseBoneNum * OutMeshBoneIndexToMirroredMeshBoneIndexMap.GetTypeSize());
+
+		// Query only bone info from mirror data table.
+		TMap<FName, FName, TInlineSetAllocator<128, FMemStackSetAllocator>> NameToMirrorNameBoneMap;
+		InMirrorDataTable.ForeachRow<FMirrorTableRow>(TEXT("UE::UAF::FillMirrorBoneIndexes"), [&NameToMirrorNameBoneMap](const FName& Key, const FMirrorTableRow& Value) mutable
+			{
+				if (Value.MirrorEntryType == EMirrorRowType::Bone)
+				{
+					NameToMirrorNameBoneMap.Add(Value.Name, Value.MirroredName);
+				}
+			}
+		);
+
+		// We need the reference skeleton to be able to query bone names.
+		// @todo: Will this change once we have abstract hierarchies?
+		const USkeletalMesh* SkeletalMesh = InReferencePose.GetSkeletalMeshAsset();
+		checkf(SkeletalMesh != nullptr, TEXT("UAF::FillMirrorBoneIndices - SkeletalMesh is null"));
+		const FReferenceSkeleton& ReferenceSkeleton = SkeletalMesh->GetRefSkeleton();
+		
+		const TArrayView<const FBoneIndexType> MeshBoneIndexMapToLODBoneIndex = InReferencePose.GetMeshBoneIndexToLODBoneIndexMap();
+		checkf(MeshBoneIndexMapToLODBoneIndex.Num() == RefPoseBoneNum, TEXT("Buffer mismatch: %d:%d"), MeshBoneIndexMapToLODBoneIndex.Num(), RefPoseBoneNum);
+
+		// Build mirror map for all the mesh bones in reference pose (LOD0).
+		for (FBoneIndexType MeshBoneIndex = 0; MeshBoneIndex < OutMeshBoneIndexToMirroredMeshBoneIndexMap.Num(); ++MeshBoneIndex)
+		{
+			if (OutMeshBoneIndexToMirroredMeshBoneIndexMap[MeshBoneIndex] == FBoneIndexType(INDEX_NONE))
+			{
+				// Find the candidate mirror partner for this bone (falling back to mirroring to self)
+				const FName SourceBoneName = ReferenceSkeleton.GetBoneName(MeshBoneIndex);
+				int32 MirrorMeshBoneIndex = INDEX_NONE;
+				
+				FName* MirroredBoneName = NameToMirrorNameBoneMap.Find(SourceBoneName);
+				if (!SourceBoneName.IsNone() && MirroredBoneName)
+				{
+					MirrorMeshBoneIndex = ReferenceSkeleton.FindBoneIndex(*MirroredBoneName);
+				}
+				
+				OutMeshBoneIndexToMirroredMeshBoneIndexMap[MeshBoneIndex] = MirrorMeshBoneIndex;
+				
+				// Map candidate mirror partner to current bone.
+				// @todo: What happens we have conflicts on mirror data table setup. Should we have a conflict policy to resolve this? For example, two bones map to another.
+				if (MirrorMeshBoneIndex != INDEX_NONE)
+				{
+					OutMeshBoneIndexToMirroredMeshBoneIndexMap[MirrorMeshBoneIndex] = MeshBoneIndex;
+				}
+			}
+		}
+	}
+
+	void MirrorPose(
+		FLODPose& InOutLODPose,
+		const UMirrorDataTable& InMirrorDataTable)
+	{
+		const int32 SourceLODBoneNum = InOutLODPose.LocalTransformsView.Num();
+		
+		if (InMirrorDataTable.MirrorAxis == EAxis::None)
+		{
+			UE_LOG(LogAnimation, Warning, TEXT("UAF::MirrorPose - No mirror axis provided."));
+			return;
+		}
+
+		if (SourceLODBoneNum == 0)
+		{
+			UE_LOG(LogAnimation, Warning, TEXT("UAF::MirrorPose - Attempting to mirror an empty pose."));
+			return;
+		}
+
+		if (InOutLODPose.IsAdditive())
+		{
+			UE_LOG(LogAnimation, Warning, TEXT("UAF::MirrorPose - Attempting to mirror an additive pose."));
+			return;
+		}
+		
+		// Get mirrored bones mapping.
+		TArray<FBoneIndexType, FAnimStackAllocator> MeshBoneIndexToMirroredMeshBoneIndexMap;
+		TArray<FQuat, FAnimStackAllocator> ComponentSpaceReferencePoseRotations;
+		TArray<FQuat, FAnimStackAllocator> ComponentSpaceReferenceRotationCorrections;
+
+		const FReferencePose& ReferencePose = InOutLODPose.GetRefPose();
+		
+		const int32 MirrorBoneIndicesNum = GetNumOfBonesForMirrorData(ReferencePose);
+		MeshBoneIndexToMirroredMeshBoneIndexMap.SetNumUninitialized(MirrorBoneIndicesNum);
+
+		const int32 BindPoseMirrorDataNum = GetNumOfBonesForMirrorData(InOutLODPose.GetRefPose());
+		ComponentSpaceReferencePoseRotations.SetNumUninitialized(BindPoseMirrorDataNum);
+		ComponentSpaceReferenceRotationCorrections.SetNumUninitialized(BindPoseMirrorDataNum);
+		
+		BuildMeshBoneIndexMirrorMap(ReferencePose, InMirrorDataTable, MeshBoneIndexToMirroredMeshBoneIndexMap);
+		BuildReferencePoseMirrorData(ReferencePose, InMirrorDataTable.MirrorAxis, MeshBoneIndexToMirroredMeshBoneIndexMap, ComponentSpaceReferencePoseRotations, ComponentSpaceReferenceRotationCorrections);
+		
+		MirrorPose(InOutLODPose, InMirrorDataTable.MirrorAxis, MeshBoneIndexToMirroredMeshBoneIndexMap, ComponentSpaceReferencePoseRotations, ComponentSpaceReferenceRotationCorrections);
+	}
+	
+	void MirrorPose(
+		FLODPose& InOutLODPose,
+		EAxis::Type InMirrorAxis,
+		TConstArrayView<FBoneIndexType> InMeshBoneIndexToMirroredMeshBoneIndexMap,
+		TConstArrayView<FQuat> InRefPoseMeshSpaceRotations,
+		TConstArrayView<FQuat> InRefPoseMeshSpaceRotationCorrections)
+	{
+		const int32 SourceLODBoneNum = InOutLODPose.LocalTransformsView.Num();
+		
+		if (InMirrorAxis == EAxis::None )
+		{
+			UE_LOG(LogAnimation, Warning, TEXT("UAF::MirrorPose - No mirror axis provided."));
+			return;
+		}
+
+		if (SourceLODBoneNum == 0)
+		{
+			UE_LOG(LogAnimation, Warning, TEXT("UAF::MirrorPose - Attempting to mirror an empty pose."));
+			return;
+		}
+
+		if (InOutLODPose.IsAdditive())
+		{
+			UE_LOG(LogAnimation, Warning, TEXT("UAF::MirrorPose - Attempting to mirror an additive pose."));
+			return;
+		}
+		
+		const TArrayView<const FBoneIndexType>& LODBoneIndexToParentLODBoneIndexMap = InOutLODPose.GetLODBoneIndexToParentLODBoneIndexMap();
+		checkf(LODBoneIndexToParentLODBoneIndexMap.Num() == InOutLODPose.GetNumBones(), TEXT("Buffer mismatch: %d:%d"), LODBoneIndexToParentLODBoneIndexMap.Num(), InOutLODPose.GetNumBones());
+		
+		// Mirroring is authored in object space and as such we must transform the local space transforms in object space in order
+		// to apply the object space mirroring axis. To facilitate this, we use object space transforms for the bind pose which can be cached.
+		// We ignore the translation/scale part of the bind pose as they don't impact mirroring.
+		// 
+		// Rotations, translations, and scales are all treated differently:
+		//    Rotation:
+		//        We transform the local space rotation into object space
+		//        We mirror the rotation axis
+		//        We apply a correction: if the source and target bones are different, we must account for the mirrored delta between the two
+		//        We transform the result back into local space
+		//    Translation:
+		//        We rotate the local space translation into object space
+		//        We mirror the result
+		//        We then rotate it back into local space
+		//    Scale:
+		//        Mirroring does not modify scale
+		// 
+		// This sadly doesn't quite work for additive poses because in order to transform it into the bind pose reference frame,
+		// we need the base pose it is applied on. Worse still, the base pose might not be static, it could be a time scaled sequence.
+		//
+		// Contract/Assumption: the (SourceBoneIndex -> TargetBoneIndex) mapping used to call this matches the mapping used to build OutMirrorCorrections (i.e., Target == MirrorMap[Source]).
+		const auto MirrorTransform = [&InRefPoseMeshSpaceRotations, &InMirrorAxis, InRefPoseMeshSpaceRotationCorrections](const FTransform& SourceTransform, const FBoneIndexType& SourceParentIndex, const FBoneIndexType& SourceBoneIndex, const FBoneIndexType& TargetParentIndex, const FBoneIndexType& TargetBoneIndex) -> FTransform
+		{
+			const FQuat TargetParentRefRotation = TargetParentIndex != FBoneIndexType(INDEX_NONE) ? InRefPoseMeshSpaceRotations[TargetParentIndex] : FQuat::Identity;
+			const FQuat SourceParentRefRotation = SourceParentIndex != FBoneIndexType(INDEX_NONE) ? InRefPoseMeshSpaceRotations[SourceParentIndex] : FQuat::Identity;
+
+			// Mirror the translation component: Rotate the translation into the space of the mirror plane, mirror across the plane, and rotate into the space of its new parent.
+
+			FVector T = SourceTransform.GetTranslation();
+			T = SourceParentRefRotation.RotateVector(T);          // to mesh space (component space)
+			T = FAnimationRuntime::MirrorVector(T, InMirrorAxis); // reflect across plane
+			T = TargetParentRefRotation.UnrotateVector(T);        // back to target's parent local space
+
+			// Mirror the rotation component: Rotate into the space of the mirror plane, mirror across the plane, apply corrective rotation to align result with target space's rest orientation, 
+			// then rotate into the space of its new parent
+
+			
+			FQuat Q = SourceTransform.GetRotation();
+			Q = SourceParentRefRotation * Q;                     // to mesh space (component space)
+			Q = FAnimationRuntime::MirrorQuat(Q, InMirrorAxis);  // reflect accros plane
+
+			// Bind alignment correction: mirror the source bone's bind orientation and align to target bone's bind orientation.
+			// @note: we use a precomputed correction quat to be able to save a MirrorQuat() since this can be cached once per skeleton,
+			// this assumes that source always is the mesh bone index and target always is the mirrored mesh bone index.
+			Q *= InRefPoseMeshSpaceRotationCorrections[SourceBoneIndex];
+			Q = TargetParentRefRotation.Inverse() * Q;
+			Q.Normalize();
+
+			// Scale is not affected by mirroring.
+			FVector S = SourceTransform.GetScale3D();
+
+			return FTransform(Q, T, S);
+		};
+		
+		// Mirror the root bone.
+		{
+			// @todo: Can the root bone be other than 0? LODPose.GetRefPose().RootBoneIndex? Can we have more than one root bone? 
+			FBoneIndexType RootBoneIndex = 0;
+			FBoneIndexType MirrorRootBoneIndex = InMeshBoneIndexToMirroredMeshBoneIndexMap[RootBoneIndex];
+
+			if (MirrorRootBoneIndex != FBoneIndexType(INDEX_NONE))
+			{
+				const FBoneIndexType TargetParentIndex = LODBoneIndexToParentLODBoneIndexMap[RootBoneIndex];
+				const FBoneIndexType SourceParentIndex = LODBoneIndexToParentLODBoneIndexMap[MirrorRootBoneIndex];
+
+				if (TargetParentIndex != FBoneIndexType(INDEX_NONE))
+				{
+					UE_LOGFMT(LogAnimation, Error, "UAF::MirrorPose - Found parent bone index ({0}) for root bone index ({1}). This is invalid.", TargetParentIndex, RootBoneIndex);
+					return;
+				}
+				
+				if (TargetParentIndex == FBoneIndexType(INDEX_NONE) && SourceParentIndex != FBoneIndexType(INDEX_NONE))
+				{
+					UE_LOGFMT(LogAnimation, Error, "UAF::MirrorPose - Mapping root bone to non-root bone is not supported ({0} -> {1})", RootBoneIndex, MirrorRootBoneIndex);
+					return;
+				}
+				
+				InOutLODPose.LocalTransformsView[RootBoneIndex] = MirrorTransform(InOutLODPose.LocalTransformsView[MirrorRootBoneIndex], SourceParentIndex, MirrorRootBoneIndex, TargetParentIndex, RootBoneIndex);
+			}
+		}
+		
+		// Mirror the non-root bones.
+		for (FBoneIndexType LODBoneIndex = 1; LODBoneIndex < SourceLODBoneNum; ++LODBoneIndex)
+		{
+			// @note: We can safely use the LODBoneIndex to get our mirrored name since the map is in LOD0 which will always contain the index.
+			const FBoneIndexType MirroredLODBoneIndex = InMeshBoneIndexToMirroredMeshBoneIndexMap[LODBoneIndex];
+			const FBoneIndexType ParentLODBoneIndex = LODBoneIndexToParentLODBoneIndexMap[LODBoneIndex];
+			
+			// Self-mirror (mirror in place relative to the same parent).
+			if (LODBoneIndex == MirroredLODBoneIndex)
+			{
+				InOutLODPose.LocalTransformsView[LODBoneIndex] = MirrorTransform(InOutLODPose.LocalTransformsView[LODBoneIndex], ParentLODBoneIndex, LODBoneIndex, ParentLODBoneIndex, LODBoneIndex);
+				continue;
+			}
+
+			// Skip invalid or already-processed mirror pairs.
+			if (MirroredLODBoneIndex == FBoneIndexType(INDEX_NONE) || MirroredLODBoneIndex < LODBoneIndex)
+			{
+				continue;
+			}
+			
+			const FBoneIndexType MirroredParentLODIndex = LODBoneIndexToParentLODBoneIndexMap[MirroredLODBoneIndex];
+
+			const FTransform OriginalTransformAtBoneIndex = InOutLODPose.LocalTransformsView[LODBoneIndex];
+			const FTransform OriginalTransformAtMirroredIndex = InOutLODPose.LocalTransformsView[MirroredLODBoneIndex];
+			
+			const FTransform NewTransformAtBoneIndex = MirrorTransform(OriginalTransformAtMirroredIndex, MirroredParentLODIndex, MirroredLODBoneIndex, ParentLODBoneIndex, LODBoneIndex);
+			const FTransform NewTransformAtMirroredIndex = MirrorTransform(OriginalTransformAtBoneIndex, ParentLODBoneIndex, LODBoneIndex, MirroredParentLODIndex, MirroredLODBoneIndex);
+			
+			InOutLODPose.LocalTransformsView[LODBoneIndex] = NewTransformAtBoneIndex;
+			InOutLODPose.LocalTransformsView[MirroredLODBoneIndex] = NewTransformAtMirroredIndex;
+		}
+	}
+}
diff --git a/Engine/Plugins/Experimental/UAF/UAFMirroring/Source/UAFMirroring/Private/Mirroring.h b/Engine/Plugins/Experimental/UAF/UAFMirroring/Source/UAFMirroring/Private/Mirroring.h
new file mode 100644
--- /dev/null
+++ b/Engine/Plugins/Experimental/UAF/UAFMirroring/Source/UAFMirroring/Private/Mirroring.h
@@ -0,0 +1,72 @@
+﻿// Copyright Epic Games, Inc. All Rights Reserved.
+
+#pragma once
+
+#include "CoreMinimal.h"
+#include "LODPose.h"
+
+namespace UE::UAF
+{
+	/** Number of mesh bones (use this when pre-sizing the arrays for bind pose mirror data and mirror map). */
+	int32 GetNumOfBonesForMirrorData(const FReferencePose & InReferencePose);
+	
+	/**
+	 * Build the bind-pose data needed to mirror a pose.
+	 * *
+	 * @param InReferencePose Reference/bind pose context. Used to access bind pose local transforms and parent maps.
+	 * @param InMirrorAxis Mirror axis (X/Y/Z)
+	 * @param InMeshBoneIndexToMirroredMeshBoneIndexMap Array view of mirrored mesh bone indices
+	 * @param OutRefPoseMeshSpaceRotations A pre-sized array view to fill reference pose rotations in mesh space
+	 * @param OutRefPoseMeshSpaceRotationCorrections A pre-sized array view to fill reference pose rotation corrections in mesh space
+	 * 
+	 * @note Output arrays are expected to be pre-sized to the number of bones in the reference pose's skeleton (i.e. LOD0).
+	 */
+	void BuildReferencePoseMirrorData(
+		const FReferencePose & InReferencePose,
+		EAxis::Type InMirrorAxis,
+		TConstArrayView<FBoneIndexType> InMeshBoneIndexToMirroredMeshBoneIndexMap,
+		TArrayView<FQuat> OutRefPoseMeshSpaceRotations,
+		TArrayView<FQuat> OutRefPoseMeshSpaceRotationCorrections);
+
+	/**
+	 * Build a mesh-bone mirror map from a mirror data table.
+	 * 
+	 * @param InReferencePose Reference/bind pose context. Used to query a mesh bone index by name.
+	 * @param InMirrorDataTable Data table with mirror bone name pairs (and mirror axis)
+	 * @param OutMeshBoneIndexToMirroredMeshBoneIndexMap A pre-sized array view to fill with mirrored mesh bone indices.
+	 * 
+	 * @note Output map is expected to be pre-sized to the number of bones in the reference pose's skeleton (i.e. LOD0)
+	 */
+	void BuildMeshBoneIndexMirrorMap(
+		const FReferencePose & InReferencePose,
+		const UMirrorDataTable& InMirrorDataTable,
+		TArrayView<FBoneIndexType> OutMeshBoneIndexToMirroredMeshBoneIndexMap);
+	
+	/**
+	 * Mirror a pose in place using a mirror data table.
+	 * 
+	 * @param InOutLODPose Pose to mirror (local-space)
+	 * @param InMirrorDataTable Data table with mirror bone name pairs (and mirror axis)
+	 */
+	void MirrorPose(
+		FLODPose& InOutLODPose,
+		const UMirrorDataTable& InMirrorDataTable);
+	
+	/**
+	 * Mirror a pose in place using explicit mirror axis + precomputed bind-pose data.
+	 * 
+	 * @param InOutLODPose Pose to mirror (local-space)
+	 * @param InMirrorAxis Mirror axis (X/Y/Z)
+	 * @param InMeshBoneIndexToMirroredMeshBoneIndexMap Mapping from mesh bone index to its mirrored counterpart (if any)
+	 * @param InRefPoseMeshSpaceRotations Reference pose rotations in mesh space
+	 * @param InRefPoseMeshSpaceRotationCorrections Reference pose rotation corrections in mesh space
+	 * 
+	 * @note Input map and arrays are expected to be the size of to the number of bones in the reference pose's skeleton (i.e. LOD0)
+	 */
+	void MirrorPose(
+		FLODPose& InOutLODPose,
+		EAxis::Type InMirrorAxis,
+		TConstArrayView<FBoneIndexType> InMeshBoneIndexToMirroredMeshBoneIndexMap,
+		TConstArrayView<FQuat> InRefPoseMeshSpaceRotations,
+		TConstArrayView<FQuat> InRefPoseMeshSpaceRotationCorrections);
+}
diff --git a/Engine/Plugins/Experimental/UAF/UAFMirroring/Source/UAFMirroring/Private/MirroringTrait.cpp b/Engine/Plugins/Experimental/UAF/UAFMirroring/Source/UAFMirroring/Private/MirroringTrait.cpp
new file mode 100644
--- /dev/null
+++ b/Engine/Plugins/Experimental/UAF/UAFMirroring/Source/UAFMirroring/Private/MirroringTrait.cpp
@@ -0,0 +1,228 @@
+// Copyright Epic Games, Inc. All Rights Reserved.
+
+#include "MirroringTrait.h"
+
+#include "GenerationTools.h"
+#include "EvaluationVM/EvaluationVM.h"
+#include "TransformArrayOperations.h"
+#include "Animation/MirrorDataTable.h"
+#include "LODPose.h"
+#include "Animation/Skeleton.h"
+#include "AnimationRuntime.h"
+#include "Mirroring.h"
+
+namespace UE::UAF
+{
+	////////////////////////////////////////////////////////////////////////////////////////////////////////////////////////
+	// FMirroringTrait
+	
+	AUTO_REGISTER_ANIM_TRAIT(FMirroringTrait)
+	
+	// Trait implementation boilerplate
+	#define TRAIT_INTERFACE_ENUMERATOR(GeneratorMacro) \
+	GeneratorMacro(IEvaluate) \
+	GeneratorMacro(IHierarchy) \
+	GeneratorMacro(IGarbageCollection) \
+	
+	GENERATE_ANIM_TRAIT_IMPLEMENTATION(FMirroringTrait, TRAIT_INTERFACE_ENUMERATOR, NULL_ANIM_TRAIT_INTERFACE_ENUMERATOR, NULL_ANIM_TRAIT_EVENT_ENUMERATOR)
+	#undef TRAIT_INTERFACE_ENUMERATOR
+
+	void FMirroringTrait::FInstanceData::Construct(const FExecutionContext& Context, const FTraitBinding& Binding)
+	{
+		FTrait::FInstanceData::Construct(Context, Binding);
+		
+		const FSharedData* SharedData = Binding.GetSharedData<FSharedData>();
+		FInstanceData* InstanceData = Binding.GetInstanceData<FInstanceData>();
+		
+		if (!InstanceData->Input.IsValid())
+		{
+			InstanceData->Input = Context.AllocateNodeInstance(Binding, SharedData->Input);
+		}
+		
+		IGarbageCollection::RegisterWithGC(Context, Binding);
+	}
+
+	void FMirroringTrait::FInstanceData::Deconstruct(const FExecutionContext& Context, const FTraitBinding& Binding)
+	{
+		FTrait::FInstanceData::Destruct(Context, Binding);
+
+		IGarbageCollection::UnregisterWithGC(Context, Binding);
+	}
+
+	void FMirroringTrait::PostEvaluate(FEvaluateTraversalContext& Context, const TTraitBinding<IEvaluate>& Binding) const
+	{
+		IEvaluate::PostEvaluate(Context, Binding);
+		
+		const FSharedData* SharedData = Binding.GetSharedData<FSharedData>();
+		FInstanceData* InstanceData = Binding.GetInstanceData<FInstanceData>();
+
+		// Copy since the settings properties are 'latent' and can change.
+		InstanceData->Settings = SharedData->GetSettings(Binding);
+		
+		if (!InstanceData->Settings.bShouldMirror || !InstanceData->Input.IsValid() || !InstanceData->Settings.MirrorDataTable)
+		{
+			return;
+		}
+		
+		Context.AppendTask(FAnimNextEvaluationMirroringTask::Make(InstanceData->Settings, &InstanceData->Cache));
+	}
+
+	uint32 FMirroringTrait::GetNumChildren(const FExecutionContext& Context, const TTraitBinding<IHierarchy>& Binding) const
+	{
+		return 1;
+	}
+
+	void FMirroringTrait::GetChildren(const FExecutionContext& Context, const TTraitBinding<IHierarchy>& Binding, FChildrenArray& Children) const
+	{
+		const FInstanceData* InstanceData = Binding.GetInstanceData<FInstanceData>();
+
+		// Add the child, even if the handle is empty
+		Children.Add(InstanceData->Input);
+	}
+
+	void FMirroringTrait::AddReferencedObjects(const FExecutionContext& Context, const TTraitBinding<IGarbageCollection>& Binding, FReferenceCollector& Collector) const
+	{
+		IGarbageCollection::AddReferencedObjects(Context, Binding, Collector);
+
+		FInstanceData* InstanceData = Binding.GetInstanceData<FInstanceData>();
+
+		Collector.AddReferencedObject(InstanceData->Settings.MirrorDataTable);
+	}
+
+	////////////////////////////////////////////////////////////////////////////////////////////////////////////////////////
+	// FMirroringAdditiveTrait
+	
+	AUTO_REGISTER_ANIM_TRAIT(FMirroringAdditiveTrait)
+	
+	// Trait implementation boilerplate
+	#define TRAIT_INTERFACE_ENUMERATOR(GeneratorMacro) \
+	GeneratorMacro(IEvaluate) \
+	GeneratorMacro(IGarbageCollection) \
+	
+	GENERATE_ANIM_TRAIT_IMPLEMENTATION(FMirroringAdditiveTrait, TRAIT_INTERFACE_ENUMERATOR, NULL_ANIM_TRAIT_INTERFACE_ENUMERATOR, NULL_ANIM_TRAIT_EVENT_ENUMERATOR)
+	#undef TRAIT_INTERFACE_ENUMERATOR
+
+	void FMirroringAdditiveTrait::FInstanceData::Construct(const FExecutionContext& Context, const FTraitBinding& Binding)
+	{
+		FTrait::FInstanceData::Construct(Context, Binding);
+
+		IGarbageCollection::RegisterWithGC(Context, Binding);
+	}
+
+	void FMirroringAdditiveTrait::FInstanceData::Deconstruct(const FExecutionContext& Context, const FTraitBinding& Binding)
+	{
+		FTrait::FInstanceData::Destruct(Context, Binding);
+
+		IGarbageCollection::UnregisterWithGC(Context, Binding);
+	}
+
+	void FMirroringAdditiveTrait::PostEvaluate(FEvaluateTraversalContext& Context, const TTraitBinding<IEvaluate>& Binding) const
+	{
+		IEvaluate::PostEvaluate(Context, Binding);
+		
+		const FSharedData* SharedData = Binding.GetSharedData<FSharedData>();
+		FInstanceData* InstanceData = Binding.GetInstanceData<FInstanceData>();
+
+		// Copy since the settings properties are 'latent' and can change.
+		InstanceData->Settings = SharedData->GetSettings(Binding);
+		
+		if (!InstanceData->Settings.bShouldMirror || !InstanceData->Settings.MirrorDataTable)
+		{
+			return;
+		}
+		
+		Context.AppendTask(FAnimNextEvaluationMirroringTask::Make(InstanceData->Settings, &InstanceData->Cache));
+	}
+
+	void FMirroringAdditiveTrait::AddReferencedObjects(const FExecutionContext& Context, const TTraitBinding<IGarbageCollection>& Binding, FReferenceCollector& Collector) const
+	{
+		IGarbageCollection::AddReferencedObjects(Context, Binding, Collector);
+
+		FInstanceData* InstanceData = Binding.GetInstanceData<FInstanceData>();
+
+		Collector.AddReferencedObject(InstanceData->Settings.MirrorDataTable);
+	}
+	
+} // namespace UE::UAF
+
+////////////////////////////////////////////////////////////////////////////////////////////////////////////////////////
+// FAnimNextEvaluationMirroringTask
+
+FAnimNextEvaluationMirroringTask FAnimNextEvaluationMirroringTask::Make(const UE::UAF::FMirroringTraitSettings & Settings, UE::UAF::FMirroringTraitCache* Cache)
+{
+	FAnimNextEvaluationMirroringTask Task;
+	Task.Settings = Settings;
+	Task.Cache = Cache;
+	return Task;
+}
+
+void FAnimNextEvaluationMirroringTask::Execute(UE::UAF::FEvaluationVM& VM) const
+{
+	QUICK_SCOPE_CYCLE_COUNTER(FAnimNextEvaluationMirroringTask_Execute);
+
+	// Pop our top pose and start processing it.
+	TUniquePtr<UE::UAF::FKeyframeState> Keyframe;
+	if (!VM.PopValue(UE::UAF::KEYFRAME_STACK_NAME, Keyframe))
+	{
+		// We have no inputs, nothing to do.
+		return;
+	}
+
+	// We don't support mirroring additive pose at the moment.
+	if (Keyframe->Pose.IsAdditive())
+	{
+		UE_LOG(LogAnimation, Warning, TEXT("FAnimNextEvaluationMirroringTask::Execute - Mirroring an additive pose is not supported."))
+		
+		// Force a bind pose to make it obvious.
+		Keyframe->Pose.SetIdentity(true);
+		VM.PushValue(UE::UAF::KEYFRAME_STACK_NAME, MoveTemp(Keyframe));
+		return;
+	}
+
+	// @todo: Once we use cache, validate data here before using it.
+	
+	// Compact pose format of Mirror Bone Map
+	TCustomBoneIndexArray<FCompactPoseBoneIndex, FCompactPoseBoneIndex> CompactPoseMirrorBones;
+	
+	// Query compact pose.
+	// @todo: Cache this and only compute it when data is invalidated.
+	if (EnumHasAnyFlags(VM.GetFlags(), UE::UAF::EEvaluationFlags::Attributes))
+	{
+		// Get required bones.
+		int32 NumRequiredBones = Keyframe->Pose.GetSkeletonAsset()->GetReferenceSkeleton().GetNum();
+		
+		TArray<FBoneIndexType> RequiredBoneIndexArray;
+		RequiredBoneIndexArray.AddUninitialized(NumRequiredBones);
+		for (int32 BoneIndex = 0; BoneIndex < RequiredBoneIndexArray.Num(); ++BoneIndex)
+		{
+			RequiredBoneIndexArray[BoneIndex] = static_cast<FBoneIndexType>(BoneIndex);
+		}
+		
+		FBoneContainer BoneContainer;
+		BoneContainer.InitializeTo(RequiredBoneIndexArray, UE::Anim::FCurveFilterSettings(UE::Anim::ECurveFilterMode::DisallowAll), *Keyframe->Pose.GetSkeletonAsset());
+		
+		TCustomBoneIndexArray<FSkeletonPoseBoneIndex, FSkeletonPoseBoneIndex> MirrorBoneIndexes;
+		Settings.MirrorDataTable->FillMirrorBoneIndexes(BoneContainer.GetSkeletonAsset(), MirrorBoneIndexes);
+		Settings.MirrorDataTable->FillCompactPoseMirrorBones(BoneContainer, MirrorBoneIndexes, CompactPoseMirrorBones);
+	}
+
+	// Mirror pose from input.
+	if (EnumHasAnyFlags(VM.GetFlags(), UE::UAF::EEvaluationFlags::Bones))
+	{
+		// @todo: Use MirrorPose() signature that leverages cached arrays to avoid computing them every time we run the task.
+		UE::UAF::MirrorPose(Keyframe->Pose, *Settings.MirrorDataTable);
+	}
+
+	if (EnumHasAnyFlags(VM.GetFlags(), UE::UAF::EEvaluationFlags::Curves))
+	{
+		FAnimationRuntime::MirrorCurves(Keyframe->Curves, *Settings.MirrorDataTable);
+	}
+
+	if (EnumHasAnyFlags(VM.GetFlags(), UE::UAF::EEvaluationFlags::Attributes))
+	{
+		// @todo: Should we make a version of MirrorAttributes() that works with FBoneIndex and/or an ArrayView? Should be fine either way once we cache CompactPoseMirrorBones.
+		UE::Anim::Attributes::MirrorAttributes(Keyframe->Attributes, *Settings.MirrorDataTable, CompactPoseMirrorBones);
+	}
+
+	VM.PushValue(UE::UAF::KEYFRAME_STACK_NAME, MoveTemp(Keyframe));
+}
diff --git a/Engine/Plugins/Experimental/UAF/UAFMirroring/Source/UAFMirroring/Private/MirroringTrait.h b/Engine/Plugins/Experimental/UAF/UAFMirroring/Source/UAFMirroring/Private/MirroringTrait.h
new file mode 100644
--- /dev/null
+++ b/Engine/Plugins/Experimental/UAF/UAFMirroring/Source/UAFMirroring/Private/MirroringTrait.h
@@ -0,0 +1,126 @@
+// Copyright Epic Games, Inc. All Rights Reserved.
+
+#pragma once
+
+#include "CoreMinimal.h"
+#include "Kismet/KismetMathLibrary.h"
+#include "TraitCore/Trait.h"
+#include "TraitCore/TraitSharedData.h"
+#include "TraitInterfaces/IAttributeProvider.h"
+#include "TraitInterfaces/IEvaluate.h"
+#include "TraitInterfaces/IHierarchy.h"
+#include "TraitInterfaces/IUpdate.h"
+#include "TraitInterfaces/IGarbageCollection.h"
+#include "MirroringTraitData.h"
+
+#include "MirroringTrait.generated.h"
+
+namespace UE::UAF
+{
+	/**
+	 * Mirroring Cache
+	 *
+	 * Holds all precomputed data needed to mirror keyframe output efficiently, so it doesn’t have to be rebuilt every evaluation.
+	 * 
+	 * @see FAnimNextEvaluationMirroringTask
+	 */
+	struct FMirroringTraitCache
+	{
+		// Cached mirror indices map
+		TArray<FBoneIndexType> MeshBoneIndexToMirroredMeshBoneIndexMap;
+		
+		// Cached bind pose rotations
+		TArray<FQuat> ComponentSpaceReferencePoseRotations;
+
+		// Cached bind pose rotation corrections
+		TArray<FQuat> ComponentSpaceReferenceRotationCorrection;
+
+		// @todo: we still need to detect skeleton and mirror data table changes. Maybe cache prev used here.
+	};
+	
+	/**
+	 * Mirroring Base Trait
+	 * 
+	 * A trait that can mirror an input's keyframe data.
+	 */
+	struct FMirroringTrait : FBaseTrait, IEvaluate, IHierarchy, IGarbageCollection // @todo: IInertializerBlend to support inertial blends.
+	{
+		DECLARE_ANIM_TRAIT(FMirroringTrait, FBaseTrait)
+
+		using FSharedData = FMirroringTraitSharedData;
+
+		struct FInstanceData : FTrait::FInstanceData
+		{
+			FTraitPtr Input;
+			FMirroringTraitSettings Settings;
+			FMirroringTraitCache Cache;
+			
+			void Construct(const FExecutionContext& Context, const FTraitBinding& Binding);
+			void Deconstruct(const FExecutionContext& Context, const FTraitBinding& Binding);
+		};
+		
+		// IEvaluate impl
+		virtual void PostEvaluate(FEvaluateTraversalContext& Context, const TTraitBinding<IEvaluate>& Binding) const override;
+		
+		// IHierarchy impl
+		virtual uint32 GetNumChildren(const FExecutionContext& Context, const TTraitBinding<IHierarchy>& Binding) const override;
+		virtual void GetChildren(const FExecutionContext& Context, const TTraitBinding<IHierarchy>& Binding, FChildrenArray& Children) const override;
+		
+		// IGarbageCollection impl
+		virtual void AddReferencedObjects(const FExecutionContext& Context, const TTraitBinding<IGarbageCollection>& Binding, FReferenceCollector& Collector) const override;
+	};
+
+	/**
+	 * Mirroring Additive Trait
+	 * 
+	 * Same behaviour as FMirroringTrait, but as an additive (i.e. it only mirrors the super-trait’s output).
+	 */
+	struct FMirroringAdditiveTrait : FAdditiveTrait, IEvaluate, IGarbageCollection // @todo: IInertializerBlend to support inertial blends.
+	{
+		DECLARE_ANIM_TRAIT(FMirroringAdditiveTrait, FAdditiveTrait)
+
+		using FSharedData = FMirroringAdditiveTraitSharedData;
+
+		struct FInstanceData : FTrait::FInstanceData
+		{
+			FMirroringTraitSettings Settings;
+			FMirroringTraitCache Cache;
+
+			void Construct(const FExecutionContext& Context, const FTraitBinding& Binding);
+			void Deconstruct(const FExecutionContext& Context, const FTraitBinding& Binding);
+		};
+
+		// IEvaluate impl
+		virtual void PostEvaluate(FEvaluateTraversalContext& Context, const TTraitBinding<IEvaluate>& Binding) const override;
+
+		// IGarbageCollection impl
+		virtual void AddReferencedObjects(const FExecutionContext& Context, const TTraitBinding<IGarbageCollection>& Binding, FReferenceCollector& Collector) const override;
+	};
+}
+
+/**
+ * Mirroring Task
+ * 
+ * This pop the top keyframe from the VM keyframe stack, mirrors it, and pushes
+ * back the result onto the stack.
+ */
+USTRUCT()
+struct FAnimNextEvaluationMirroringTask : public FAnimNextEvaluationTask
+{
+	GENERATED_BODY()
+
+	DECLARE_ANIM_EVALUATION_TASK(FAnimNextEvaluationMirroringTask)
+
+	// Make a FAnimNextDeadBlendingApplyTask task
+	static FAnimNextEvaluationMirroringTask Make(const UE::UAF::FMirroringTraitSettings & Settings, UE::UAF::FMirroringTraitCache* Cache);
+	
+	// Settings to use when doing the mirror pass
+	UPROPERTY(VisibleAnywhere, Category = "Settings", meta=(ShowOnlyInnerProperties))
+	FUAFMirroringTraitSettings Settings;
+
+	// Cache for bind pose rotation, corrections, and mirror map
+	UE::UAF::FMirroringTraitCache* Cache = nullptr;
+	
+	// Task entry point
+	virtual void Execute(UE::UAF::FEvaluationVM& VM) const override;
+};
\ No newline at end of file
diff --git a/Engine/Plugins/Experimental/UAF/UAFMirroring/Source/UAFMirroring/Private/Module.cpp b/Engine/Plugins/Experimental/UAF/UAFMirroring/Source/UAFMirroring/Private/Module.cpp
new file mode 100644
--- /dev/null
+++ b/Engine/Plugins/Experimental/UAF/UAFMirroring/Source/UAFMirroring/Private/Module.cpp
@@ -0,0 +1,33 @@
+﻿// Copyright Epic Games, Inc. All Rights Reserved.
+
+#include "CoreMinimal.h"
+#include "Animation/MirrorDataTable.h"
+#include "Module/AnimNextModule.h"
+#include "Modules/ModuleInterface.h"
+#include "Modules/ModuleManager.h"
+#include "RigVMCore/RigVMRegistry.h"
+
+namespace UE::UAF::Mirroring
+{
+
+class FModule : public IModuleInterface
+{
+public:
+	virtual void StartupModule() override
+	{
+		static TPair<UClass*, FRigVMRegistry::ERegisterObjectOperation> const AllowedObjectTypes[] =
+		{
+			{ UMirrorDataTable::StaticClass(), FRigVMRegistry::ERegisterObjectOperation::Class },
+		};
+
+		FRigVMRegistry::Get().RegisterObjectTypes(AllowedObjectTypes);
+	}
+
+	virtual void ShutdownModule() override
+	{
+	}
+};
+
+}
+
+IMPLEMENT_MODULE(UE::UAF::Mirroring::FModule, UAFMirroring)
\ No newline at end of file
diff --git a/Engine/Plugins/Experimental/UAF/UAFMirroring/Source/UAFMirroring/UAFMirroring.Build.cs b/Engine/Plugins/Experimental/UAF/UAFMirroring/Source/UAFMirroring/UAFMirroring.Build.cs
new file mode 100644
--- /dev/null
+++ b/Engine/Plugins/Experimental/UAF/UAFMirroring/Source/UAFMirroring/UAFMirroring.Build.cs
@@ -0,0 +1,55 @@
+﻿// Copyright Epic Games, Inc. All Rights Reserved.
+
+using UnrealBuildTool;
+
+public class UAFMirroring : ModuleRules
+{
+	public UAFMirroring(ReadOnlyTargetRules Target) : base(Target)
+	{
+		PCHUsage = ModuleRules.PCHUsageMode.UseExplicitOrSharedPCHs;
+		
+		PublicIncludePaths.AddRange(
+			new string[] {
+				// ... add public include paths required here ...
+			}
+			);
+				
+		
+		PrivateIncludePaths.AddRange(
+			new string[] {
+				// ... add other private include paths required here ...
+			}
+			);
+			
+		
+		PublicDependencyModuleNames.AddRange(
+			new string[]
+			{
+				// ... add other public dependencies that you statically link with here ...
+			}
+			);
+			
+		
+		PrivateDependencyModuleNames.AddRange(
+			new string[]
+			{
+				"Core",
+				"CoreUObject",
+				"AnimationCore",
+				"RigVM",
+				"Engine",
+				"UAF",
+				"UAFAnimGraph",
+				// ... add private dependencies that you statically link with here ...	
+			}
+			);
+		
+		
+		DynamicallyLoadedModuleNames.AddRange(
+			new string[]
+			{
+				// ... add any modules that your module loads dynamically here ...
+			}
+			);
+	}
+}
diff --git a/Engine/Plugins/Experimental/UAF/UAFMirroring/Source/UAFMirroringUncookedOnly/Private/UAFGraphNodeTemplate_Mirror.h b/Engine/Plugins/Experimental/UAF/UAFMirroring/Source/UAFMirroringUncookedOnly/Private/UAFGraphNodeTemplate_Mirror.h
new file mode 100644
--- /dev/null
+++ b/Engine/Plugins/Experimental/UAF/UAFMirroring/Source/UAFMirroringUncookedOnly/Private/UAFGraphNodeTemplate_Mirror.h
@@ -0,0 +1,46 @@
+﻿// Copyright Epic Games, Inc. All Rights Reserved.
+
+#pragma once
+
+#include "AnimGraphUncookedOnlyUtils.h"
+#include "AnimNextController.h"
+#include "Styling/AppStyle.h"
+#include "Styling/SlateIconFinder.h"
+#include "Templates/UAFGraphNodeTemplate.h"
+#include "MirroringTraitData.h"
+#include "Animation/MirrorDataTable.h"
+
+#include "UAFGraphNodeTemplate_Mirror.generated.h"
+
+#define LOCTEXT_NAMESPACE "UAFGraphNodeTemplate_Mirror"
+
+UCLASS()
+class UUAFGraphNodeTemplate_Mirror : public UUAFGraphNodeTemplate
+{
+	GENERATED_BODY()
+
+	UUAFGraphNodeTemplate_Mirror()
+	{
+		Title = LOCTEXT("MirrorTitle", "Mirror");
+		TooltipText = LOCTEXT("MirrorTooltip", "Mirror an input using a mirror data table");
+		Category = LOCTEXT("MirrorCategory", "UAF");
+		MenuDescription = LOCTEXT("MirrorMenuDesc", "Mirror");
+		Color = FLinearColor(FColor(62, 140, 35)); // From UAssetDefinition_DataTable
+		Icon = *FSlateIconFinder::FindIconForClass(UMirrorDataTable::StaticClass()).GetIcon();
+		Traits =
+		{
+			TInstancedStruct<UE::UAF::FMirroringTraitData>::Make(),
+		};
+		DragDropAssetTypes.Add(UMirrorDataTable::StaticClass());
+		SetCategoryForPinsInLayout(
+			{
+				GET_PIN_PATH_STRING_CHECKED(FMirroringTraitSharedData, Input),
+				GET_PIN_PATH_STRING_CHECKED(FMirroringTraitSharedData, Settings),
+			},
+			FRigVMPinCategory::GetDefaultCategoryName(),
+			NodeLayout,
+			true);
+	}
+};
+
+#undef LOCTEXT_NAMESPACE
\ No newline at end of file
diff --git a/Engine/Plugins/Experimental/UAF/UAFMirroring/Source/UAFMirroringUncookedOnly/Private/UAFMirroringUncookedOnlyModule.cpp b/Engine/Plugins/Experimental/UAF/UAFMirroring/Source/UAFMirroringUncookedOnly/Private/UAFMirroringUncookedOnlyModule.cpp
new file mode 100644
--- /dev/null
+++ b/Engine/Plugins/Experimental/UAF/UAFMirroring/Source/UAFMirroringUncookedOnly/Private/UAFMirroringUncookedOnlyModule.cpp
@@ -0,0 +1,6 @@
+﻿// Copyright Epic Games, Inc. All Rights Reserved.
+
+#include "CoreMinimal.h"
+#include "Modules/ModuleManager.h"
+
+IMPLEMENT_MODULE(FDefaultModuleImpl, UAFMirroringUncookedOnly);
\ No newline at end of file
diff --git a/Engine/Plugins/Experimental/UAF/UAFMirroring/Source/UAFMirroringUncookedOnly/UAFMirroringUncookedOnly.Build.cs b/Engine/Plugins/Experimental/UAF/UAFMirroring/Source/UAFMirroringUncookedOnly/UAFMirroringUncookedOnly.Build.cs
new file mode 100644
--- /dev/null
+++ b/Engine/Plugins/Experimental/UAF/UAFMirroring/Source/UAFMirroringUncookedOnly/UAFMirroringUncookedOnly.Build.cs
@@ -0,0 +1,28 @@
+﻿// Copyright Epic Games, Inc. All Rights Reserved.
+
+using UnrealBuildTool;
+
+public class UAFMirroringUncookedOnly : ModuleRules
+{
+    public UAFMirroringUncookedOnly(ReadOnlyTargetRules Target) : base(Target)
+    {
+	    PrivateDependencyModuleNames.AddRange(
+		    new string[]
+		    {
+			    "Core",
+			    "CoreUObject",
+			    "AnimationCore",
+			    "Engine",
+			    "UAF",
+			    "UAFUncookedOnly",
+			    "UAFAnimGraph",
+			    "UAFAnimGraphUncookedOnly",
+			    "UAFMirroring",
+			    "RigVMDeveloper",
+			    "RigVM",
+			    "Slate",
+			    "SlateCore",
+		    }
+	    );
+    }
+}
\ No newline at end of file
diff --git a/Engine/Plugins/Experimental/UAF/UAFMirroring/UAFMirroring.uplugin b/Engine/Plugins/Experimental/UAF/UAFMirroring/UAFMirroring.uplugin
new file mode 100644
--- /dev/null
+++ b/Engine/Plugins/Experimental/UAF/UAFMirroring/UAFMirroring.uplugin
@@ -0,0 +1,42 @@
+﻿{
+	"FileVersion": 3,
+	"Version": 1,
+	"VersionName": "0.1",
+	"FriendlyName": "UAF Mirroring",
+	"Description": "Keyframe mirroring for UAF",
+	"Category": "Other",
+	"CreatedBy": "Epic Games, Inc.",
+	"CreatedByURL": "https://epicgames.com",
+	"DocsURL": "",
+	"MarketplaceURL": "",
+	"CanContainContent": false,
+	"IsBetaVersion": false,
+	"IsExperimentalVersion": true,
+	"Installed": false,
+	"Plugins": [
+		{
+			"Name": "UAF",
+			"Enabled": true
+		},
+		{
+			"Name": "UAFAnimGraph",
+			"Enabled": true
+		},
+		{
+			"Name": "RigVM",
+			"Enabled": true
+		}
+	],
+	"Modules": [
+		{
+			"Name": "UAFMirroring",
+			"Type": "Runtime",
+			"LoadingPhase": "Default"
+		},
+		{
+			"Name": "UAFMirroringUncookedOnly",
+			"Type": "UncookedOnly",
+			"LoadingPhase": "Default"
+		}
+	]
+}
\ No newline at end of file
