UAF Mirroring improvements:

- Now we cache mesh space rotations and corrections to avoid recomputing them every evaluation pass.
- Now we cache mirror bone mapping to avoid recomputing them every evaluation pass.
- Fixed trait parameters structs not being registed.
- Fixed channel flags not being properly used to determine what to mirror.
- Added size checks in MirrorPose.

[RN] minor

diff --git a/Engine/Plugins/Experimental/UAF/UAFMirroring/Source/UAFMirroring/Internal/MirroringTraitData.h b/Engine/Plugins/Experimental/UAF/UAFMirroring/Source/UAFMirroring/Internal/MirroringTraitData.h
--- a/Engine/Plugins/Experimental/UAF/UAFMirroring/Source/UAFMirroring/Internal/MirroringTraitData.h
+++ b/Engine/Plugins/Experimental/UAF/UAFMirroring/Source/UAFMirroring/Internal/MirroringTraitData.h
@@ -4,19 +4,20 @@
 
 #include "TraitCore/TraitBinding.h"
 #include "TraitCore/TraitSharedData.h"
-#include "Animation/MirrorDataTable.h"
 
 #include "MirroringTraitData.generated.h"
 
+class UMirrorDataTable;
+
 /**
- * Mirroring Settings
+ * Mirroring Setup Parameters
  *
- * This struct holds all the information need to mirror keyframe data (i.e. Bones, Curves, Attributes)
+ * This struct holds the minimal configuration needed to enable a mirror pass.
  *
  * @see UMirrorDataTable, FMirroringTrait, FMirroringAdditiveTrait
  */
-USTRUCT(BlueprintType, meta = (DisplayName = "Mirroring Trait Settings"))
-struct FUAFMirroringTraitSettings
+USTRUCT(BlueprintType, meta = (DisplayName = "Mirroring Trait Setup Parameters"))
+struct FUAFMirroringTraitSetupParams
 {
 	GENERATED_BODY()
 
@@ -24,23 +25,34 @@ struct FUAFMirroringTraitSettings
 	UPROPERTY(EditAnywhere, Category = Settings, meta=(DisplayName = "Mirror"))
 	bool bShouldMirror = true;
 
+	// Data table to map bones to their mirrored counterpart
+	UPROPERTY(EditAnywhere, Category = Settings, meta = (ExportAsReference="true"))
+	TObjectPtr<const UMirrorDataTable> MirrorDataTable = nullptr;
+};
+
+/**
+ * Mirroring Apply/Filter Parameters
+ *
+ * This struct holds flags for the channels that can be affected during the mirror pass (i.e. Bones, Curves, Attributes).
+ *
+ * @see FMirroringTrait, FMirroringAdditiveTrait, FAnimNextEvaluationMirroringTask
+ */
+USTRUCT(BlueprintType, meta = (DisplayName = "Mirroring Trait Apply To Parameters"))
+struct FUAFMirroringTraitApplyToParams
+{
+	GENERATED_BODY()
+
 	// Whether to mirror bone transforms
-	UPROPERTY(EditAnywhere, Category = Channels, meta=(DisplayName = "Bones", EditCondition="bShouldMirror", EditConditionHides))
+	UPROPERTY(EditAnywhere, Category = Channels, meta=(DisplayName = "Bones"))
 	bool bShouldMirrorBones = true;
 
 	// Whether to mirror animation curves
-	UPROPERTY(EditAnywhere, Category = Channels, meta=(DisplayName = "Curves", EditCondition="bShouldMirror", EditConditionHides))
+	UPROPERTY(EditAnywhere, Category = Channels, meta=(DisplayName = "Curves"))
 	bool bShouldMirrorCurves = true;
 
 	// Whether to mirror attributes
-	UPROPERTY(EditAnywhere, Category = Channels, meta=(DisplayName = "Attributes", EditCondition="bShouldMirror", EditConditionHides))
+	UPROPERTY(EditAnywhere, Category = Channels, meta=(DisplayName = "Attributes"))
 	bool bShouldMirrorAttributes = true;
-
-	// Data table to map bones to their mirrored counterpart
-	UPROPERTY(EditAnywhere, Category = Settings, meta = (ExportAsReference="true"))
-	TObjectPtr<const UMirrorDataTable> MirrorDataTable = nullptr;
-
-	// @todo: We need transitions settings (to properly handle DeadBlending).
 };
 
 /** A trait that can mirror an input's keyframe data. */
@@ -48,17 +60,24 @@ USTRUCT(meta = (DisplayName = "Mirroring", ShowTooltip = true))
 struct FMirroringTraitSharedData : public FAnimNextTraitSharedData
 {
 	GENERATED_BODY()
-	
+
+	// Input node to query the keyframe to be mirrored.
 	UPROPERTY()
 	FAnimNextTraitHandle Input;
 
-	UPROPERTY(EditAnywhere, Category = "Settings", meta=(ShowOnlyInnerProperties))
-	FUAFMirroringTraitSettings Settings;
+	// Defines whether to perform mirror pass (and what data table to use).
+	UPROPERTY(EditAnywhere, Category = "Setup", meta=(ShowOnlyInnerProperties))
+	FUAFMirroringTraitSetupParams Setup;
+
+	// Defines what channels will be affected during the mirror pass.
+	UPROPERTY(EditAnywhere, Category = "Apply To", meta=(ShowOnlyInnerProperties))
+	FUAFMirroringTraitApplyToParams ApplyTo;
 	
 	// Latent pin support boilerplate
 	#define TRAIT_LATENT_PROPERTIES_ENUMERATOR(GeneratorMacro) \
 	GeneratorMacro(Input) \
-	GeneratorMacro(Settings) \
+	GeneratorMacro(Setup) \
+	GeneratorMacro(ApplyTo) \
 
 	GENERATE_TRAIT_LATENT_PROPERTIES(FMirroringTraitSharedData, TRAIT_LATENT_PROPERTIES_ENUMERATOR)
 	#undef TRAIT_LATENT_PROPERTIES_ENUMERATOR
@@ -70,12 +89,18 @@ struct FMirroringAdditiveTraitSharedData : public FAnimNextTraitSharedData
 {
 	GENERATED_BODY()
 
-	UPROPERTY(EditAnywhere, Category = "Settings", meta=(ShowOnlyInnerProperties))
-	FUAFMirroringTraitSettings Settings;
+	// Defines whether to perform mirror pass (and what data table to use).
+	UPROPERTY(EditAnywhere, Category = "Setup", meta=(ShowOnlyInnerProperties))
+	FUAFMirroringTraitSetupParams Setup;
 
+	// Defines what channels will be affected during the mirror pass.
+	UPROPERTY(EditAnywhere, Category = "Apply To", meta=(ShowOnlyInnerProperties))
+	FUAFMirroringTraitApplyToParams ApplyTo;
+	
 	// Latent pin support boilerplate
 	#define TRAIT_LATENT_PROPERTIES_ENUMERATOR(GeneratorMacro) \
-	GeneratorMacro(Settings) \
+	GeneratorMacro(Setup) \
+	GeneratorMacro(ApplyTo) \
 
 	GENERATE_TRAIT_LATENT_PROPERTIES(FMirroringAdditiveTraitSharedData, TRAIT_LATENT_PROPERTIES_ENUMERATOR)
 	#undef TRAIT_LATENT_PROPERTIES_ENUMERATOR
@@ -84,7 +109,8 @@ struct FMirroringAdditiveTraitSharedData : public FAnimNextTraitSharedData
 namespace UE::UAF
 {
 	// Namespaced alias
-	using FMirroringTraitSettings = FUAFMirroringTraitSettings;
+	using FMirroringTraitSetupParams = FUAFMirroringTraitSetupParams;
+	using FMirroringTraitApplyToParams = FUAFMirroringTraitApplyToParams;
 	using FMirroringTraitData = FMirroringTraitSharedData;
 	using FMirroringAdditiveTraitData = FMirroringAdditiveTraitSharedData;
 	
diff --git a/Engine/Plugins/Experimental/UAF/UAFMirroring/Source/UAFMirroring/Private/Mirroring.cpp b/Engine/Plugins/Experimental/UAF/UAFMirroring/Source/UAFMirroring/Private/Mirroring.cpp
--- a/Engine/Plugins/Experimental/UAF/UAFMirroring/Source/UAFMirroring/Private/Mirroring.cpp
+++ b/Engine/Plugins/Experimental/UAF/UAFMirroring/Source/UAFMirroring/Private/Mirroring.cpp
@@ -21,7 +21,7 @@ namespace UE::UAF
 	}
 	
 	void BuildReferencePoseMirrorData(
-		const FReferencePose & InReferencePose,
+		const FReferencePose& InReferencePose,
 		EAxis::Type InMirrorAxis,
 		TConstArrayView<FBoneIndexType> InMeshBoneIndexToMirroredMeshBoneIndexMap,
 		TArrayView<FQuat> OutRefPoseMeshSpaceRotations,
@@ -75,7 +75,7 @@ namespace UE::UAF
 	
 	// Get a map from source bone index to mirrored bone index.
 	void BuildMeshBoneIndexMirrorMap(
-		const FReferencePose & InReferencePose,
+		const FReferencePose& InReferencePose,
 		const UMirrorDataTable& InMirrorDataTable,
 		TArrayView<FBoneIndexType> OutMeshBoneIndexToMirroredMeshBoneIndexMap)
 	{
@@ -188,8 +188,13 @@ namespace UE::UAF
 		TConstArrayView<FQuat> InRefPoseMeshSpaceRotations,
 		TConstArrayView<FQuat> InRefPoseMeshSpaceRotationCorrections)
 	{
+		// @todo: check mirror map size
+		
 		const int32 SourceLODBoneNum = InOutLODPose.LocalTransformsView.Num();
 		
+		const int32 RefPoseBoneNum = InOutLODPose.GetRefPose().ReferenceLocalTransforms.Num();
+		checkf(RefPoseBoneNum == InMeshBoneIndexToMirroredMeshBoneIndexMap.Num(), TEXT("Buffer mismatch: %d:%d"), RefPoseBoneNum, InMeshBoneIndexToMirroredMeshBoneIndexMap.Num());
+		
 		if (InMirrorAxis == EAxis::None )
 		{
 			UE_LOG(LogAnimation, Warning, TEXT("UAF::MirrorPose - No mirror axis provided."));
diff --git a/Engine/Plugins/Experimental/UAF/UAFMirroring/Source/UAFMirroring/Private/MirroringTrait.cpp b/Engine/Plugins/Experimental/UAF/UAFMirroring/Source/UAFMirroring/Private/MirroringTrait.cpp
--- a/Engine/Plugins/Experimental/UAF/UAFMirroring/Source/UAFMirroring/Private/MirroringTrait.cpp
+++ b/Engine/Plugins/Experimental/UAF/UAFMirroring/Source/UAFMirroring/Private/MirroringTrait.cpp
@@ -8,11 +8,83 @@
 #include "Animation/MirrorDataTable.h"
 #include "LODPose.h"
 #include "Animation/Skeleton.h"
+#include "Engine/SkeletalMesh.h"
 #include "AnimationRuntime.h"
 #include "Mirroring.h"
+#include "BoneContainer.h"
 
 namespace UE::UAF
 {
+	////////////////////////////////////////////////////////////////////////////////////////////////////////////////////////
+	// FMirroringTraitCache
+
+	bool FMirroringTraitCache::AreMirrorMapsValid(const UE::UAF::FReferencePose& InReferencePose, const TWeakObjectPtr<const UMirrorDataTable>& InMirrorTable) const
+	{
+		// Ensure assets are the same.
+		
+		if (InMirrorTable != MirrorTable)
+		{
+			return false;
+		}
+
+		if (InReferencePose.SkeletalMesh != SkeletalMesh)
+		{
+			return false;
+		}
+
+		// Ensure arrays match mesh bone count.
+		
+		const int32 ExpectedNumBones = UE::UAF::GetNumOfBonesForMirrorData(InReferencePose);
+
+		const bool bAreMirrorMapsEquallySized = (ExpectedNumBones == MeshBoneIndexToMirroredMeshBoneIndexMap.Num()) && (ExpectedNumBones == CompactPoseBoneIndexToMirroredCompactPoseBoneIndexMap.Num());
+		if (!bAreMirrorMapsEquallySized)
+		{
+			return false;
+		}
+
+		const bool bIsRefPoseDataEquallySized = (ExpectedNumBones == MeshSpaceReferencePoseRotations.Num()) && (ExpectedNumBones == MeshSpaceReferenceRotationCorrections.Num());
+		if (!bIsRefPoseDataEquallySized)
+		{
+			return false;
+		}
+		
+		return true;
+	}
+	
+	bool FMirroringTraitCache::IsReferencePoseDataValid(const UE::UAF::FReferencePose& InReferencePose) const
+	{
+		// Ensure skeletal mesh is still the same.
+		if (InReferencePose.SkeletalMesh != SkeletalMesh)
+		{
+			return false;
+		}
+
+		// Ensure arrays match mesh bone count.
+		const int32 ExpectedNumBones = UE::UAF::GetNumOfBonesForMirrorData(InReferencePose);
+		const bool bIsRefPoseDataEquallySized = (ExpectedNumBones == MeshSpaceReferencePoseRotations.Num()) && (ExpectedNumBones == MeshSpaceReferenceRotationCorrections.Num());
+		if (!bIsRefPoseDataEquallySized)
+		{
+			return false;
+		}
+		
+		return true;
+	}
+
+	bool FMirroringTraitCache::IsValid(const UE::UAF::FReferencePose& InReferencePose, const TWeakObjectPtr<const UMirrorDataTable>& InMirrorTable) const
+	{
+		return AreMirrorMapsValid(InReferencePose, InMirrorTable) && IsReferencePoseDataValid(InReferencePose);
+	}
+
+	void FMirroringTraitCache::Clear()
+	{
+		MeshBoneIndexToMirroredMeshBoneIndexMap.Empty();
+		CompactPoseBoneIndexToMirroredCompactPoseBoneIndexMap.Empty();
+		MeshSpaceReferencePoseRotations.Empty();
+		MeshSpaceReferencePoseRotations.Empty();
+		MirrorTable.Reset();
+		SkeletalMesh.Reset();
+	}
+
 	////////////////////////////////////////////////////////////////////////////////////////////////////////////////////////
 	// FMirroringTrait
 	
@@ -56,15 +128,19 @@ namespace UE::UAF
 		const FSharedData* SharedData = Binding.GetSharedData<FSharedData>();
 		FInstanceData* InstanceData = Binding.GetInstanceData<FInstanceData>();
 
-		// Copy since the settings properties are 'latent' and can change.
-		InstanceData->Settings = SharedData->GetSettings(Binding);
-		
-		if (!InstanceData->Settings.bShouldMirror || !InstanceData->Input.IsValid() || !InstanceData->Settings.MirrorDataTable)
+		// Copy since the settings properties are 'latent' and can change..
+		const FMirroringTraitSetupParams LatentSetupValue = SharedData->GetSetup(Binding);;
+
+		InstanceData->bHasMirrorStateChanged = LatentSetupValue.bShouldMirror != InstanceData->Setup.bShouldMirror;
+		InstanceData->Setup = LatentSetupValue;
+		InstanceData->ApplyTo = SharedData->GetApplyTo(Binding);
+
+		if (!InstanceData->Setup.bShouldMirror || !InstanceData->Input.IsValid() || !InstanceData->Setup.MirrorDataTable)
 		{
 			return;
 		}
 		
-		Context.AppendTask(FAnimNextEvaluationMirroringTask::Make(InstanceData->Settings, &InstanceData->Cache));
+		Context.AppendTask(FAnimNextEvaluationMirroringTask::Make(InstanceData->Setup, InstanceData->ApplyTo, &InstanceData->Cache));
 	}
 
 	uint32 FMirroringTrait::GetNumChildren(const FExecutionContext& Context, const TTraitBinding<IHierarchy>& Binding) const
@@ -86,7 +162,7 @@ namespace UE::UAF
 
 		FInstanceData* InstanceData = Binding.GetInstanceData<FInstanceData>();
 
-		Collector.AddReferencedObject(InstanceData->Settings.MirrorDataTable);
+		Collector.AddReferencedObject(InstanceData->Setup.MirrorDataTable);
 	}
 
 	////////////////////////////////////////////////////////////////////////////////////////////////////////////////////////
@@ -124,14 +200,18 @@ namespace UE::UAF
 		FInstanceData* InstanceData = Binding.GetInstanceData<FInstanceData>();
 
 		// Copy since the settings properties are 'latent' and can change.
-		InstanceData->Settings = SharedData->GetSettings(Binding);
+		const FMirroringTraitSetupParams LatentSetupValue = SharedData->GetSetup(Binding);;
+
+		InstanceData->bHasMirrorStateChanged = LatentSetupValue.bShouldMirror != InstanceData->Setup.bShouldMirror;
+		InstanceData->Setup = LatentSetupValue;
+		InstanceData->ApplyTo = SharedData->GetApplyTo(Binding);
 		
-		if (!InstanceData->Settings.bShouldMirror || !InstanceData->Settings.MirrorDataTable)
+		if (!InstanceData->Setup.bShouldMirror || !InstanceData->Setup.MirrorDataTable)
 		{
 			return;
 		}
 		
-		Context.AppendTask(FAnimNextEvaluationMirroringTask::Make(InstanceData->Settings, &InstanceData->Cache));
+		Context.AppendTask(FAnimNextEvaluationMirroringTask::Make(InstanceData->Setup, InstanceData->ApplyTo, &InstanceData->Cache));
 	}
 
 	void FMirroringAdditiveTrait::AddReferencedObjects(const FExecutionContext& Context, const TTraitBinding<IGarbageCollection>& Binding, FReferenceCollector& Collector) const
@@ -140,7 +220,7 @@ namespace UE::UAF
 
 		FInstanceData* InstanceData = Binding.GetInstanceData<FInstanceData>();
 
-		Collector.AddReferencedObject(InstanceData->Settings.MirrorDataTable);
+		Collector.AddReferencedObject(InstanceData->Setup.MirrorDataTable);
 	}
 	
 } // namespace UE::UAF
@@ -148,10 +228,11 @@ namespace UE::UAF
 ////////////////////////////////////////////////////////////////////////////////////////////////////////////////////////
 // FAnimNextEvaluationMirroringTask
 
-FAnimNextEvaluationMirroringTask FAnimNextEvaluationMirroringTask::Make(const UE::UAF::FMirroringTraitSettings & Settings, UE::UAF::FMirroringTraitCache* Cache)
+FAnimNextEvaluationMirroringTask FAnimNextEvaluationMirroringTask::Make(const UE::UAF::FMirroringTraitSetupParams& Setup, const UE::UAF::FMirroringTraitApplyToParams& ApplyTo, UE::UAF::FMirroringTraitCache* Cache)
 {
 	FAnimNextEvaluationMirroringTask Task;
-	Task.Settings = Settings;
+	Task.Setup = Setup;
+	Task.ApplyTo = ApplyTo;
 	Task.Cache = Cache;
 	return Task;
 }
@@ -179,50 +260,120 @@ void FAnimNextEvaluationMirroringTask::Execute(UE::UAF::FEvaluationVM& VM) const
 		return;
 	}
 
-	// @todo: Once we use cache, validate data here before using it.
-	
-	// Compact pose format of Mirror Bone Map
-	TCustomBoneIndexArray<FCompactPoseBoneIndex, FCompactPoseBoneIndex> CompactPoseMirrorBones;
+	// Rebuild mirror cache (if needed)
+	UE::UAF::FReferencePose ReferencePose = Keyframe->Pose.GetRefPose();
+	EnsureCache(VM, ReferencePose);
 	
-	// Query compact pose.
-	// @todo: Cache this and only compute it when data is invalidated.
-	if (EnumHasAnyFlags(VM.GetFlags(), UE::UAF::EEvaluationFlags::Attributes))
+	// After attempting to rebuild cache if its still invalid something went wrong.
+	if (!Cache->IsValid(ReferencePose, Setup.MirrorDataTable))
 	{
-		// Get required bones.
-		int32 NumRequiredBones = Keyframe->Pose.GetSkeletonAsset()->GetReferenceSkeleton().GetNum();
-		
-		TArray<FBoneIndexType> RequiredBoneIndexArray;
-		RequiredBoneIndexArray.AddUninitialized(NumRequiredBones);
-		for (int32 BoneIndex = 0; BoneIndex < RequiredBoneIndexArray.Num(); ++BoneIndex)
-		{
-			RequiredBoneIndexArray[BoneIndex] = static_cast<FBoneIndexType>(BoneIndex);
-		}
-		
-		FBoneContainer BoneContainer;
-		BoneContainer.InitializeTo(RequiredBoneIndexArray, UE::Anim::FCurveFilterSettings(UE::Anim::ECurveFilterMode::DisallowAll), *Keyframe->Pose.GetSkeletonAsset());
-		
-		TCustomBoneIndexArray<FSkeletonPoseBoneIndex, FSkeletonPoseBoneIndex> MirrorBoneIndexes;
-		Settings.MirrorDataTable->FillMirrorBoneIndexes(BoneContainer.GetSkeletonAsset(), MirrorBoneIndexes);
-		Settings.MirrorDataTable->FillCompactPoseMirrorBones(BoneContainer, MirrorBoneIndexes, CompactPoseMirrorBones);
+		// Force a bind pose to make it obvious.
+		Keyframe->Pose.SetIdentity(Keyframe->Pose.IsAdditive());
+		VM.PushValue(UE::UAF::KEYFRAME_STACK_NAME, MoveTemp(Keyframe));
+		return;
 	}
 
-	// Mirror pose from input.
-	if (EnumHasAnyFlags(VM.GetFlags(), UE::UAF::EEvaluationFlags::Bones))
+	// Mirror keyframe data.
+	bool bShouldMirrorBones = ApplyTo.bShouldMirrorBones && EnumHasAnyFlags(VM.GetFlags(), UE::UAF::EEvaluationFlags::Bones);
+	bool bShouldMirrorCurves= ApplyTo.bShouldMirrorCurves && EnumHasAnyFlags(VM.GetFlags(), UE::UAF::EEvaluationFlags::Curves);
+	bool bShouldMirrorAttributes = ApplyTo.bShouldMirrorAttributes && EnumHasAnyFlags(VM.GetFlags(), UE::UAF::EEvaluationFlags::Attributes);
+	
+	if (bShouldMirrorBones)
 	{
-		// @todo: Use MirrorPose() signature that leverages cached arrays to avoid computing them every time we run the task.
-		UE::UAF::MirrorPose(Keyframe->Pose, *Settings.MirrorDataTable);
+		UE::UAF::MirrorPose(
+			Keyframe->Pose,
+			Setup.MirrorDataTable->MirrorAxis,
+			Cache->MeshBoneIndexToMirroredMeshBoneIndexMap,
+			Cache->MeshSpaceReferencePoseRotations,
+			Cache->MeshSpaceReferenceRotationCorrections);
 	}
 
-	if (EnumHasAnyFlags(VM.GetFlags(), UE::UAF::EEvaluationFlags::Curves))
+	if (bShouldMirrorCurves)
 	{
-		FAnimationRuntime::MirrorCurves(Keyframe->Curves, *Settings.MirrorDataTable);
+		FAnimationRuntime::MirrorCurves(Keyframe->Curves, *Setup.MirrorDataTable);
 	}
 
-	if (EnumHasAnyFlags(VM.GetFlags(), UE::UAF::EEvaluationFlags::Attributes))
+	if (bShouldMirrorAttributes)
 	{
-		// @todo: Should we make a version of MirrorAttributes() that works with FBoneIndex and/or an ArrayView? Should be fine either way once we cache CompactPoseMirrorBones.
-		UE::Anim::Attributes::MirrorAttributes(Keyframe->Attributes, *Settings.MirrorDataTable, CompactPoseMirrorBones);
+		// @todo: Make UAF version of MirrorAttributes() or address internal heap allocations for scratch array.
+		UE::Anim::Attributes::MirrorAttributes(
+			Keyframe->Attributes,
+			*Setup.MirrorDataTable,
+			Cache->CompactPoseBoneIndexToMirroredCompactPoseBoneIndexMap);
 	}
 
 	VM.PushValue(UE::UAF::KEYFRAME_STACK_NAME, MoveTemp(Keyframe));
 }
+
+void FAnimNextEvaluationMirroringTask::EnsureCache(const UE::UAF::FEvaluationVM& VM, const UE::UAF::FReferencePose& InReferencePose) const
+{
+	checkf(Cache != nullptr, TEXT("FAnimNextEvaluationMirroringTask::UpdateCache - Invalid cache."))
+	
+	if (InReferencePose.SkeletalMesh.IsValid())
+	{
+		const int32 NumBonesForLOD0 = UE::UAF::GetNumOfBonesForMirrorData(InReferencePose);
+		const bool bAreMirrorMapsValid = Cache->AreMirrorMapsValid(InReferencePose, Setup.MirrorDataTable);
+		const bool bIsReferencePoseDataValid = Cache->IsReferencePoseDataValid(InReferencePose);
+		const bool bShouldMirrorBones = ApplyTo.bShouldMirrorBones && EnumHasAnyFlags(VM.GetFlags(), UE::UAF::EEvaluationFlags::Bones);
+		const bool bShouldMirrorAttributes = ApplyTo.bShouldMirrorAttributes && EnumHasAnyFlags(VM.GetFlags(), UE::UAF::EEvaluationFlags::Attributes);
+		
+		if (bShouldMirrorBones)
+		{
+			if (!bAreMirrorMapsValid)
+			{
+				Cache->MeshBoneIndexToMirroredMeshBoneIndexMap.SetNumUninitialized(NumBonesForLOD0);
+				
+				UE::UAF::BuildMeshBoneIndexMirrorMap(InReferencePose, *Setup.MirrorDataTable, Cache->MeshBoneIndexToMirroredMeshBoneIndexMap);
+			}
+			
+			if (!bIsReferencePoseDataValid)
+			{
+				Cache->MeshSpaceReferencePoseRotations.SetNumUninitialized(NumBonesForLOD0);
+				Cache->MeshSpaceReferenceRotationCorrections.SetNumUninitialized(NumBonesForLOD0);
+				
+				UE::UAF::BuildReferencePoseMirrorData(
+					InReferencePose,
+					Setup.MirrorDataTable->MirrorAxis,
+					Cache->MeshBoneIndexToMirroredMeshBoneIndexMap,
+					Cache->MeshSpaceReferencePoseRotations,
+					Cache->MeshSpaceReferenceRotationCorrections
+				);
+			}
+		}
+		
+		if (bShouldMirrorAttributes)
+		{
+			if (!bAreMirrorMapsValid)
+			{
+				// We only care about the mesh bones (LOD0).
+				TArray<FBoneIndexType> RequiredBoneIndexArray;
+				RequiredBoneIndexArray.AddUninitialized(NumBonesForLOD0);
+				
+				for (int32 BoneIndex = 0; BoneIndex < RequiredBoneIndexArray.Num(); ++BoneIndex)
+				{
+					RequiredBoneIndexArray[BoneIndex] = InReferencePose.GetSkeletonBoneIndexFromLODBoneIndex(BoneIndex);
+				}
+
+				// Compute mirror map for compact pose.
+				
+				FBoneContainer BoneContainer;
+				BoneContainer.InitializeTo(RequiredBoneIndexArray, UE::Anim::FCurveFilterSettings(UE::Anim::ECurveFilterMode::DisallowAll), *InReferencePose.SkeletalMesh);
+				
+				TCustomBoneIndexArray<FSkeletonPoseBoneIndex, FSkeletonPoseBoneIndex> MirrorBoneIndexes;
+				Setup.MirrorDataTable->FillMirrorBoneIndexes(BoneContainer.GetSkeletonAsset(), MirrorBoneIndexes);
+				Setup.MirrorDataTable->FillCompactPoseMirrorBones(BoneContainer, MirrorBoneIndexes, Cache->CompactPoseBoneIndexToMirroredCompactPoseBoneIndexMap);
+			}
+		}
+	}
+	else
+	{
+		// Clear data to terminate task early.
+		Cache->Clear();
+		
+		UE_LOG(LogAnimation, Warning, TEXT("FAnimNextEvaluationMirroringTask::UpdateCache - Failed to get skeletal mesh asset from reference pose."))
+	}
+
+	// Keep track of latest assets used to build the cache
+	Cache->SkeletalMesh = InReferencePose.SkeletalMesh;
+	Cache->MirrorTable = Setup.MirrorDataTable;
+}
\ No newline at end of file
diff --git a/Engine/Plugins/Experimental/UAF/UAFMirroring/Source/UAFMirroring/Private/MirroringTrait.h b/Engine/Plugins/Experimental/UAF/UAFMirroring/Source/UAFMirroring/Private/MirroringTrait.h
--- a/Engine/Plugins/Experimental/UAF/UAFMirroring/Source/UAFMirroring/Private/MirroringTrait.h
+++ b/Engine/Plugins/Experimental/UAF/UAFMirroring/Source/UAFMirroring/Private/MirroringTrait.h
@@ -12,9 +12,13 @@
 #include "TraitInterfaces/IUpdate.h"
 #include "TraitInterfaces/IGarbageCollection.h"
 #include "MirroringTraitData.h"
+#include "CustomBoneIndexArray.h"
 
 #include "MirroringTrait.generated.h"
 
+class UMirrorDataTable;
+class USkeletalMesh;
+
 namespace UE::UAF
 {
 	/**
@@ -26,16 +30,44 @@ namespace UE::UAF
 	 */
 	struct FMirroringTraitCache
 	{
-		// Cached mirror indices map
+		// Cached mirror indices map (invalidated on skeletal mesh or mirror data table change)
 		TArray<FBoneIndexType> MeshBoneIndexToMirroredMeshBoneIndexMap;
 		
-		// Cached bind pose rotations
-		TArray<FQuat> ComponentSpaceReferencePoseRotations;
-
-		// Cached bind pose rotation corrections
-		TArray<FQuat> ComponentSpaceReferenceRotationCorrection;
+		// Cached bind pose rotations (invalidated on skeletal mesh change)
+		TArray<FQuat> MeshSpaceReferencePoseRotations;
 
-		// @todo: we still need to detect skeleton and mirror data table changes. Maybe cache prev used here.
+		// Cached bind pose rotation corrections (invalidated on skeletal mesh change)
+		TArray<FQuat> MeshSpaceReferenceRotationCorrections;
+		
+		// Cached mirror indices map but using FCompactPoseIndex (invalidated on skeletal mesh or mirror data table change)
+		TCustomBoneIndexArray<FCompactPoseBoneIndex, FCompactPoseBoneIndex> CompactPoseBoneIndexToMirroredCompactPoseBoneIndexMap;
+
+		// Skeletal mesh used to build this cache
+		TWeakObjectPtr<const USkeletalMesh> SkeletalMesh = nullptr;
+
+		// Mirror data table used to build this cache
+		TWeakObjectPtr<const UMirrorDataTable> MirrorTable = nullptr;
+
+		/**
+		 * True if the cached mirror maps were generated using exactly these assets.
+		 * Any mismatch (different mesh/mirror table, size mismatch) means you must rebuild the mirror maps.
+		 */
+		bool AreMirrorMapsValid(const UE::UAF::FReferencePose& InReferencePose, const TWeakObjectPtr<const UMirrorDataTable>& InMirrorTable) const;
+
+		/**
+		 * True if the cached bind/reference-pose data matches this mesh.
+		 * Any mismatch (different mesh table, size mismatch)  means you must rebuild the reference-pose arrays.
+		 */
+		bool IsReferencePoseDataValid(const UE::UAF::FReferencePose& InReferencePose) const;
+
+		/**
+		 * True if the cache was generated using exactly these assets.
+		 * Any mismatch (different mesh/mirror table, size mismatch) means you must rebuild the cache.
+		 */
+		bool IsValid(const UE::UAF::FReferencePose& InReferencePose, const TWeakObjectPtr<const UMirrorDataTable>& InMirrorTable) const;
+
+		/** Empty arrays and reset ptrs to assets */
+		void Clear();
 	};
 	
 	/**
@@ -43,7 +75,7 @@ namespace UE::UAF
 	 * 
 	 * A trait that can mirror an input's keyframe data.
 	 */
-	struct FMirroringTrait : FBaseTrait, IEvaluate, IHierarchy, IGarbageCollection // @todo: IInertializerBlend to support inertial blends.
+	struct FMirroringTrait : FBaseTrait, IEvaluate, IHierarchy, IGarbageCollection
 	{
 		DECLARE_ANIM_TRAIT(FMirroringTrait, FBaseTrait)
 
@@ -51,8 +83,19 @@ namespace UE::UAF
 
 		struct FInstanceData : FTrait::FInstanceData
 		{
+			// Input node to query the keyframe to be mirrored.
 			FTraitPtr Input;
-			FMirroringTraitSettings Settings;
+			
+			// Defines whether to perform mirror pass (and what data table to use).
+			FMirroringTraitSetupParams Setup;
+
+			// Defines what channels will be affected during the mirror pass.
+			FMirroringTraitApplyToParams ApplyTo;
+
+			// Used to trigger inertial blends.
+			bool bHasMirrorStateChanged = false;
+			
+			// Data that stays the same between evaluation passes.
 			FMirroringTraitCache Cache;
 			
 			void Construct(const FExecutionContext& Context, const FTraitBinding& Binding);
@@ -75,7 +118,7 @@ namespace UE::UAF
 	 * 
 	 * Same behaviour as FMirroringTrait, but as an additive (i.e. it only mirrors the super-trait’s output).
 	 */
-	struct FMirroringAdditiveTrait : FAdditiveTrait, IEvaluate, IGarbageCollection // @todo: IInertializerBlend to support inertial blends.
+	struct FMirroringAdditiveTrait : FAdditiveTrait, IUpdate, IEvaluate, IGarbageCollection
 	{
 		DECLARE_ANIM_TRAIT(FMirroringAdditiveTrait, FAdditiveTrait)
 
@@ -83,7 +126,16 @@ namespace UE::UAF
 
 		struct FInstanceData : FTrait::FInstanceData
 		{
-			FMirroringTraitSettings Settings;
+			// Defines whether to perform mirror pass (and what data table to use).
+			FMirroringTraitSetupParams Setup;
+
+			// Defines what channels will be affected during the mirror pass.
+			FMirroringTraitApplyToParams ApplyTo;
+			
+			// Used to trigger inertial blends.
+			bool bHasMirrorStateChanged = false;
+			
+			// Data that stays the same between evaluation passes.
 			FMirroringTraitCache Cache;
 
 			void Construct(const FExecutionContext& Context, const FTraitBinding& Binding);
@@ -96,7 +148,8 @@ namespace UE::UAF
 		// IGarbageCollection impl
 		virtual void AddReferencedObjects(const FExecutionContext& Context, const TTraitBinding<IGarbageCollection>& Binding, FReferenceCollector& Collector) const override;
 	};
-}
+	
+} // namespace UE::UAF
 
 /**
  * Mirroring Task
@@ -111,16 +164,24 @@ struct FAnimNextEvaluationMirroringTask : public FAnimNextEvaluationTask
 
 	DECLARE_ANIM_EVALUATION_TASK(FAnimNextEvaluationMirroringTask)
 
-	// Make a FAnimNextDeadBlendingApplyTask task
-	static FAnimNextEvaluationMirroringTask Make(const UE::UAF::FMirroringTraitSettings & Settings, UE::UAF::FMirroringTraitCache* Cache);
+	// Make a FAnimNextEvaluationMirroringTask task
+	static FAnimNextEvaluationMirroringTask Make(const UE::UAF::FMirroringTraitSetupParams & Setup, const UE::UAF::FMirroringTraitApplyToParams & ApplyTo, UE::UAF::FMirroringTraitCache* Cache);
 	
 	// Settings to use when doing the mirror pass
-	UPROPERTY(VisibleAnywhere, Category = "Settings", meta=(ShowOnlyInnerProperties))
-	FUAFMirroringTraitSettings Settings;
+	UPROPERTY(VisibleAnywhere, Category = "Setup")
+	FUAFMirroringTraitSetupParams Setup;
+
+	UPROPERTY(VisibleAnywhere, Category = "Apply To")
+	FUAFMirroringTraitApplyToParams ApplyTo;
 
 	// Cache for bind pose rotation, corrections, and mirror map
 	UE::UAF::FMirroringTraitCache* Cache = nullptr;
 	
 	// Task entry point
 	virtual void Execute(UE::UAF::FEvaluationVM& VM) const override;
+	
+private:
+
+	// Ensures the mirror cache matches the given inputs, rebuilds if stale.
+	void EnsureCache(const UE::UAF::FEvaluationVM& VM, const UE::UAF::FReferencePose& InReferencePose) const;
 };
\ No newline at end of file
diff --git a/Engine/Plugins/Experimental/UAF/UAFMirroring/Source/UAFMirroring/Private/Module.cpp b/Engine/Plugins/Experimental/UAF/UAFMirroring/Source/UAFMirroring/Private/Module.cpp
--- a/Engine/Plugins/Experimental/UAF/UAFMirroring/Source/UAFMirroring/Private/Module.cpp
+++ b/Engine/Plugins/Experimental/UAF/UAFMirroring/Source/UAFMirroring/Private/Module.cpp
@@ -1,6 +1,7 @@
 ﻿// Copyright Epic Games, Inc. All Rights Reserved.
 
 #include "CoreMinimal.h"
+#include "MirroringTraitData.h"
 #include "Animation/MirrorDataTable.h"
 #include "Module/AnimNextModule.h"
 #include "Modules/ModuleInterface.h"
@@ -20,7 +21,14 @@ public:
 			{ UMirrorDataTable::StaticClass(), FRigVMRegistry::ERegisterObjectOperation::Class },
 		};
 
+		static UScriptStruct* const AllowedStructTypes[] =
+		{
+			FMirroringTraitSetupParams::StaticStruct(),
+			FMirroringTraitApplyToParams::StaticStruct(),
+		};
+		
 		FRigVMRegistry::Get().RegisterObjectTypes(AllowedObjectTypes);
+		FRigVMRegistry::Get().RegisterStructTypes(AllowedStructTypes);
 	}
 
 	virtual void ShutdownModule() override
diff --git a/Engine/Plugins/Experimental/UAF/UAFMirroring/Source/UAFMirroringUncookedOnly/Private/UAFGraphNodeTemplate_Mirror.h b/Engine/Plugins/Experimental/UAF/UAFMirroring/Source/UAFMirroringUncookedOnly/Private/UAFGraphNodeTemplate_Mirror.h
--- a/Engine/Plugins/Experimental/UAF/UAFMirroring/Source/UAFMirroringUncookedOnly/Private/UAFGraphNodeTemplate_Mirror.h
+++ b/Engine/Plugins/Experimental/UAF/UAFMirroring/Source/UAFMirroringUncookedOnly/Private/UAFGraphNodeTemplate_Mirror.h
@@ -35,12 +35,44 @@ class UUAFGraphNodeTemplate_Mirror : public UUAFGraphNodeTemplate
 		SetCategoryForPinsInLayout(
 			{
 				GET_PIN_PATH_STRING_CHECKED(FMirroringTraitSharedData, Input),
-				GET_PIN_PATH_STRING_CHECKED(FMirroringTraitSharedData, Settings),
+				GET_PIN_PATH_STRING_CHECKED(FMirroringTraitSharedData, Setup),
+				GET_PIN_PATH_STRING_CHECKED(FMirroringTraitSharedData, ApplyTo),
 			},
 			FRigVMPinCategory::GetDefaultCategoryName(),
 			NodeLayout,
 			true);
 	}
+
+	virtual void HandleAssetDropped_Implementation(UAnimNextController* Controller, URigVMUnitNode* Node, UObject* Asset) const
+	{
+		Super::HandleAssetDropped_Implementation(Controller, Node, Asset);
+
+		// @todo: Wish we could use a subtitle instead but looks like the subtitle can't be set via the controller as its not instanced per not but rather queried per template.
+		
+		if (Asset)
+		{
+			Controller->OpenUndoBracket(LOCTEXT("ConfigureNodeOnDrop", "Configure Node On Drop").ToString());
+
+			Controller->SetNodeTitle(Node, FText::Format(LOCTEXT("NodeTitleFormat", "Mirror using {0}"), FText::FromString(Asset->GetName())).ToString(), true, true, true);
+			
+			Controller->CloseUndoBracket();
+		}
+	}
+
+	virtual void HandlePinDefaultValueChanged_Implementation(UAnimNextController* Controller, URigVMPin* Pin) const
+	{
+		Super::HandlePinDefaultValueChanged_Implementation(Controller, Pin);
+
+		if (Pin->GetFName() == GET_MEMBER_NAME_CHECKED(UE::UAF::FMirroringTraitSetupParams, MirrorDataTable))
+		{
+			Controller->OpenUndoBracket(LOCTEXT("SetNodeTitle", "Set Node Title").ToString());
+
+			FSoftObjectPath Path(Pin->GetDefaultValue());
+			Controller->SetNodeTitle(Pin->GetNode(), FText::Format(LOCTEXT("NodeTitleFormat", "Mirror using {0}"), FText::FromString(Path.GetAssetName())).ToString(), true, true, true);
+
+			Controller->CloseUndoBracket();
+		}
+	}
 };
 
 #undef LOCTEXT_NAMESPACE
\ No newline at end of file
