Fixed aborting mirror task entirely when skipping mirroring attribues or bones but still have other channels marking for processing. Other miscellaneous fixes.

[RN] minor

diff --git a/Engine/Plugins/Experimental/UAF/UAFMirroring/Source/UAFMirroring/Private/MirroringTrait.cpp b/Engine/Plugins/Experimental/UAF/UAFMirroring/Source/UAFMirroring/Private/MirroringTrait.cpp
--- a/Engine/Plugins/Experimental/UAF/UAFMirroring/Source/UAFMirroring/Private/MirroringTrait.cpp
+++ b/Engine/Plugins/Experimental/UAF/UAFMirroring/Source/UAFMirroring/Private/MirroringTrait.cpp
@@ -18,7 +18,7 @@ namespace UE::UAF
 	////////////////////////////////////////////////////////////////////////////////////////////////////////////////////////
 	// FMirroringTraitCache
 
-	bool FMirroringTraitCache::AreMirrorMapsValid(const UE::UAF::FReferencePose& InReferencePose, const TWeakObjectPtr<const UMirrorDataTable>& InMirrorTable) const
+	bool FMirroringTraitCache::AreMirrorMapsValid(const UE::UAF::FReferencePose& InReferencePose, const TWeakObjectPtr<const UMirrorDataTable>& InMirrorTable, bool bInShouldMirrorBones, bool bInShouldMirrorAttributes) const
 	{
 		// Ensure assets are the same.
 		
@@ -36,13 +36,20 @@ namespace UE::UAF
 		
 		const int32 ExpectedNumBones = UE::UAF::GetNumOfBonesForMirrorData(InReferencePose);
 
-		const bool bAreMirrorMapsEquallySized = (ExpectedNumBones == MeshBoneIndexToMirroredMeshBoneIndexMap.Num()) && (ExpectedNumBones == CompactPoseBoneIndexToMirroredCompactPoseBoneIndexMap.Num());
+		// Default to true if we are not mirroring bones since these arrays are only used for that channel.
+		const bool bDoesMeshMirrorMapHaveExpectedNumBones = (ExpectedNumBones == MeshBoneIndexToMirroredMeshBoneIndexMap.Num() || !bInShouldMirrorBones);
+
+		// Default to true if we are not mirroring attributes since these arrays are only used for that channel.
+		const bool bDoesCompactPoseMirrorMapHaveExpectedNumBones = (ExpectedNumBones == CompactPoseBoneIndexToMirroredCompactPoseBoneIndexMap.Num() || !bInShouldMirrorAttributes);
+		
+		const bool bAreMirrorMapsEquallySized = bDoesMeshMirrorMapHaveExpectedNumBones && bDoesCompactPoseMirrorMapHaveExpectedNumBones;
 		if (!bAreMirrorMapsEquallySized)
 		{
 			return false;
 		}
 
-		const bool bIsRefPoseDataEquallySized = (ExpectedNumBones == MeshSpaceReferencePoseRotations.Num()) && (ExpectedNumBones == MeshSpaceReferenceRotationCorrections.Num());
+		// Default to true if we are not mirroring bones since these arrays are only used for that channel.
+		const bool bIsRefPoseDataEquallySized = (ExpectedNumBones == MeshSpaceReferencePoseRotations.Num() && ExpectedNumBones == MeshSpaceReferenceRotationCorrections.Num()) || !bInShouldMirrorBones;
 		if (!bIsRefPoseDataEquallySized)
 		{
 			return false;
@@ -51,7 +58,7 @@ namespace UE::UAF
 		return true;
 	}
 	
-	bool FMirroringTraitCache::IsReferencePoseDataValid(const UE::UAF::FReferencePose& InReferencePose) const
+	bool FMirroringTraitCache::IsReferencePoseDataValid(const UE::UAF::FReferencePose& InReferencePose, bool bInShouldMirrorBones) const
 	{
 		// Ensure skeletal mesh is still the same.
 		if (InReferencePose.SkeletalMesh != SkeletalMesh)
@@ -60,8 +67,9 @@ namespace UE::UAF
 		}
 
 		// Ensure arrays match mesh bone count.
+		// Default to true if we are not mirroring bones since these arrays are only used for that channel.
 		const int32 ExpectedNumBones = UE::UAF::GetNumOfBonesForMirrorData(InReferencePose);
-		const bool bIsRefPoseDataEquallySized = (ExpectedNumBones == MeshSpaceReferencePoseRotations.Num()) && (ExpectedNumBones == MeshSpaceReferenceRotationCorrections.Num());
+		const bool bIsRefPoseDataEquallySized = (ExpectedNumBones == MeshSpaceReferencePoseRotations.Num() && ExpectedNumBones == MeshSpaceReferenceRotationCorrections.Num()) || !bInShouldMirrorBones;
 		if (!bIsRefPoseDataEquallySized)
 		{
 			return false;
@@ -70,9 +78,9 @@ namespace UE::UAF
 		return true;
 	}
 
-	bool FMirroringTraitCache::IsValid(const UE::UAF::FReferencePose& InReferencePose, const TWeakObjectPtr<const UMirrorDataTable>& InMirrorTable) const
+	bool FMirroringTraitCache::IsValid(const UE::UAF::FReferencePose& InReferencePose, const TWeakObjectPtr<const UMirrorDataTable>& InMirrorTable, bool bInShouldMirrorBones, bool bInShouldMirrorAttributes) const
 	{
-		return AreMirrorMapsValid(InReferencePose, InMirrorTable) && IsReferencePoseDataValid(InReferencePose);
+		return AreMirrorMapsValid(InReferencePose, InMirrorTable, bInShouldMirrorBones, bInShouldMirrorAttributes) && IsReferencePoseDataValid(InReferencePose, bInShouldMirrorBones);
 	}
 
 	void FMirroringTraitCache::Clear()
@@ -80,7 +88,7 @@ namespace UE::UAF
 		MeshBoneIndexToMirroredMeshBoneIndexMap.Empty();
 		CompactPoseBoneIndexToMirroredCompactPoseBoneIndexMap.Empty();
 		MeshSpaceReferencePoseRotations.Empty();
-		MeshSpaceReferencePoseRotations.Empty();
+		MeshSpaceReferenceRotationCorrections.Empty();
 		MirrorTable.Reset();
 		SkeletalMesh.Reset();
 	}
@@ -263,9 +271,14 @@ void FAnimNextEvaluationMirroringTask::Execute(UE::UAF::FEvaluationVM& VM) const
 	// Rebuild mirror cache (if needed)
 	UE::UAF::FReferencePose ReferencePose = Keyframe->Pose.GetRefPose();
 	EnsureCache(VM, ReferencePose);
+
+	// Determine what channel to skip during this pass.
+	bool bShouldMirrorBones = ApplyTo.bShouldMirrorBones && EnumHasAnyFlags(VM.GetFlags(), UE::UAF::EEvaluationFlags::Bones);
+	bool bShouldMirrorCurves= ApplyTo.bShouldMirrorCurves && EnumHasAnyFlags(VM.GetFlags(), UE::UAF::EEvaluationFlags::Curves);
+	bool bShouldMirrorAttributes = ApplyTo.bShouldMirrorAttributes && EnumHasAnyFlags(VM.GetFlags(), UE::UAF::EEvaluationFlags::Attributes);
 	
 	// After attempting to rebuild cache if its still invalid something went wrong.
-	if (!Cache->IsValid(ReferencePose, Setup.MirrorDataTable))
+	if (!Cache->IsValid(ReferencePose, Setup.MirrorDataTable, bShouldMirrorBones, bShouldMirrorAttributes))
 	{
 		// Force a bind pose to make it obvious.
 		Keyframe->Pose.SetIdentity(Keyframe->Pose.IsAdditive());
@@ -274,9 +287,6 @@ void FAnimNextEvaluationMirroringTask::Execute(UE::UAF::FEvaluationVM& VM) const
 	}
 
 	// Mirror keyframe data.
-	bool bShouldMirrorBones = ApplyTo.bShouldMirrorBones && EnumHasAnyFlags(VM.GetFlags(), UE::UAF::EEvaluationFlags::Bones);
-	bool bShouldMirrorCurves= ApplyTo.bShouldMirrorCurves && EnumHasAnyFlags(VM.GetFlags(), UE::UAF::EEvaluationFlags::Curves);
-	bool bShouldMirrorAttributes = ApplyTo.bShouldMirrorAttributes && EnumHasAnyFlags(VM.GetFlags(), UE::UAF::EEvaluationFlags::Attributes);
 	
 	if (bShouldMirrorBones)
 	{
@@ -307,15 +317,15 @@ void FAnimNextEvaluationMirroringTask::Execute(UE::UAF::FEvaluationVM& VM) const
 
 void FAnimNextEvaluationMirroringTask::EnsureCache(const UE::UAF::FEvaluationVM& VM, const UE::UAF::FReferencePose& InReferencePose) const
 {
-	checkf(Cache != nullptr, TEXT("FAnimNextEvaluationMirroringTask::UpdateCache - Invalid cache."))
+	checkf(Cache != nullptr, TEXT("FAnimNextEvaluationMirroringTask::EnsureCache - Invalid cache ptr."))
 	
 	if (InReferencePose.SkeletalMesh.IsValid())
 	{
 		const int32 NumBonesForLOD0 = UE::UAF::GetNumOfBonesForMirrorData(InReferencePose);
-		const bool bAreMirrorMapsValid = Cache->AreMirrorMapsValid(InReferencePose, Setup.MirrorDataTable);
-		const bool bIsReferencePoseDataValid = Cache->IsReferencePoseDataValid(InReferencePose);
 		const bool bShouldMirrorBones = ApplyTo.bShouldMirrorBones && EnumHasAnyFlags(VM.GetFlags(), UE::UAF::EEvaluationFlags::Bones);
 		const bool bShouldMirrorAttributes = ApplyTo.bShouldMirrorAttributes && EnumHasAnyFlags(VM.GetFlags(), UE::UAF::EEvaluationFlags::Attributes);
+		const bool bAreMirrorMapsValid = Cache->AreMirrorMapsValid(InReferencePose, Setup.MirrorDataTable, bShouldMirrorBones, bShouldMirrorAttributes);
+		const bool bIsReferencePoseDataValid = Cache->IsReferencePoseDataValid(InReferencePose, bShouldMirrorBones);
 		
 		if (bShouldMirrorBones)
 		{
@@ -370,7 +380,7 @@ void FAnimNextEvaluationMirroringTask::EnsureCache(const UE::UAF::FEvaluationVM&
 		// Clear data to terminate task early.
 		Cache->Clear();
 		
-		UE_LOG(LogAnimation, Warning, TEXT("FAnimNextEvaluationMirroringTask::UpdateCache - Failed to get skeletal mesh asset from reference pose."))
+		UE_LOG(LogAnimation, Warning, TEXT("FAnimNextEvaluationMirroringTask::EnsureCache - Failed to get skeletal mesh asset from reference pose."))
 	}
 
 	// Keep track of latest assets used to build the cache
diff --git a/Engine/Plugins/Experimental/UAF/UAFMirroring/Source/UAFMirroring/Private/MirroringTrait.h b/Engine/Plugins/Experimental/UAF/UAFMirroring/Source/UAFMirroring/Private/MirroringTrait.h
--- a/Engine/Plugins/Experimental/UAF/UAFMirroring/Source/UAFMirroring/Private/MirroringTrait.h
+++ b/Engine/Plugins/Experimental/UAF/UAFMirroring/Source/UAFMirroring/Private/MirroringTrait.h
@@ -52,19 +52,19 @@ namespace UE::UAF
 		 * True if the cached mirror maps were generated using exactly these assets.
 		 * Any mismatch (different mesh/mirror table, size mismatch) means you must rebuild the mirror maps.
 		 */
-		bool AreMirrorMapsValid(const UE::UAF::FReferencePose& InReferencePose, const TWeakObjectPtr<const UMirrorDataTable>& InMirrorTable) const;
+		bool AreMirrorMapsValid(const UE::UAF::FReferencePose& InReferencePose, const TWeakObjectPtr<const UMirrorDataTable>& InMirrorTable, bool bInShouldMirrorBones, bool bInShouldMirrorAttributes) const;
 
 		/**
 		 * True if the cached bind/reference-pose data matches this mesh.
 		 * Any mismatch (different mesh table, size mismatch)  means you must rebuild the reference-pose arrays.
 		 */
-		bool IsReferencePoseDataValid(const UE::UAF::FReferencePose& InReferencePose) const;
+		bool IsReferencePoseDataValid(const UE::UAF::FReferencePose& InReferencePose, bool bInShouldMirrorBones) const;
 
 		/**
 		 * True if the cache was generated using exactly these assets.
 		 * Any mismatch (different mesh/mirror table, size mismatch) means you must rebuild the cache.
 		 */
-		bool IsValid(const UE::UAF::FReferencePose& InReferencePose, const TWeakObjectPtr<const UMirrorDataTable>& InMirrorTable) const;
+		bool IsValid(const UE::UAF::FReferencePose& InReferencePose, const TWeakObjectPtr<const UMirrorDataTable>& InMirrorTable, bool bInShouldMirrorBones, bool bInShouldMirrorAttributes) const;
 
 		/** Empty arrays and reset ptrs to assets */
 		void Clear();
