Fixed UAF Mirroring LOD issues

[RN] minor

#Synced-CL 48298300

diff --git a/Engine/Plugins/Experimental/UAF/UAFMirroring/Source/UAFMirroring/Private/Mirroring.cpp b/Engine/Plugins/Experimental/UAF/UAFMirroring/Source/UAFMirroring/Private/Mirroring.cpp
--- a/Engine/Plugins/Experimental/UAF/UAFMirroring/Source/UAFMirroring/Private/Mirroring.cpp
+++ b/Engine/Plugins/Experimental/UAF/UAFMirroring/Source/UAFMirroring/Private/Mirroring.cpp
@@ -42,24 +42,24 @@ namespace UE::UAF
 
 		// Convert the local ref pose rotations to component space.
 		// @note: We assume the root bone is at index 0.
-		for (int32 RefBoneIndex = 1; RefBoneIndex < RefPoseBoneNum; ++RefBoneIndex)
+		for (int32 MeshBoneIndex = 1; MeshBoneIndex < RefPoseBoneNum; ++MeshBoneIndex)
 		{
-			const FBoneIndexType ParentBoneIndex = MeshBoneIndexToParentMeshBoneIndexMap[RefBoneIndex];
-			
-			const FQuat ParentLocalSpaceRotation = OutRefPoseMeshSpaceRotations[ParentBoneIndex];
-			const FQuat RefBoneLocalSpaceRotation = OutRefPoseMeshSpaceRotations[RefBoneIndex];
+			const FBoneIndexType LODBoneIndex = InReferencePose.MeshBoneIndexToLODBoneIndexMap[MeshBoneIndex];
+			const FBoneIndexType ParentMeshBoneIndex = MeshBoneIndexToParentMeshBoneIndexMap[MeshBoneIndex];
 			
+			const FQuat ParentLocalSpaceRotation = OutRefPoseMeshSpaceRotations[ParentMeshBoneIndex];
+			const FQuat RefBoneLocalSpaceRotation = InReferencePose.ReferenceLocalTransforms[LODBoneIndex].GetRotation();
 			FQuat MeshSpaceRotation =  ParentLocalSpaceRotation * RefBoneLocalSpaceRotation;
 			MeshSpaceRotation.Normalize();
 			
-			OutRefPoseMeshSpaceRotations[RefBoneIndex] = MeshSpaceRotation;
+			OutRefPoseMeshSpaceRotations[MeshBoneIndex] = MeshSpaceRotation;
 		}
 
 		// Now we can precompute the corrective rotation to align result with target space's rest orientation.
 		OutRefPoseMeshSpaceRotationCorrections[0] = FQuat::Identity;
-		for (FBoneIndexType RefBoneIndex = 1; RefBoneIndex < RefPoseBoneNum; ++RefBoneIndex)
+		for (FBoneIndexType MeshBoneIndex = 1; MeshBoneIndex < RefPoseBoneNum; ++MeshBoneIndex)
 		{
-			const FBoneIndexType SourceBoneIndex = RefBoneIndex;
+			const FBoneIndexType SourceBoneIndex = MeshBoneIndex;
 			const FBoneIndexType TargetBoneIndex = InMeshBoneIndexToMirroredMeshBoneIndexMap[SourceBoneIndex];
 
 			// Not mapped, so skip.
@@ -68,8 +68,8 @@ namespace UE::UAF
 				continue;
 			}
 			
-			OutRefPoseMeshSpaceRotationCorrections[RefBoneIndex] = FAnimationRuntime::MirrorQuat(OutRefPoseMeshSpaceRotations[SourceBoneIndex], InMirrorAxis).Inverse() * OutRefPoseMeshSpaceRotations[TargetBoneIndex];
-			OutRefPoseMeshSpaceRotationCorrections[RefBoneIndex].Normalize();
+			OutRefPoseMeshSpaceRotationCorrections[MeshBoneIndex] = FAnimationRuntime::MirrorQuat(OutRefPoseMeshSpaceRotations[SourceBoneIndex], InMirrorAxis).Inverse() * OutRefPoseMeshSpaceRotations[TargetBoneIndex];
+			OutRefPoseMeshSpaceRotationCorrections[MeshBoneIndex].Normalize();
 		}
 	}
 	
@@ -107,29 +107,48 @@ namespace UE::UAF
 		checkf(SkeletalMesh != nullptr, TEXT("UAF::FillMirrorBoneIndices - SkeletalMesh is null"));
 		const FReferenceSkeleton& ReferenceSkeleton = SkeletalMesh->GetRefSkeleton();
 		
-		const TArrayView<const FBoneIndexType> MeshBoneIndexMapToLODBoneIndex = InReferencePose.GetMeshBoneIndexToLODBoneIndexMap();
-		checkf(MeshBoneIndexMapToLODBoneIndex.Num() == RefPoseBoneNum, TEXT("Buffer mismatch: %d:%d"), MeshBoneIndexMapToLODBoneIndex.Num(), RefPoseBoneNum);
+		const TArrayView<const FBoneIndexType> LODBoneIndexToSkeletonBoneIndexMap = InReferencePose.GetLODBoneIndexToSkeletonBoneIndexMap(0);
+		const TArrayView<const FBoneIndexType> MeshBoneIndexToLODBoneIndexMap = InReferencePose.GetMeshBoneIndexToLODBoneIndexMap();
+		const TMap<FName, FBoneIndexType> BoneNameToLODBoneIndexMap = InReferencePose.GetBoneNameToLODBoneIndexMap();
+		const TArrayView<const FBoneIndexType> LODBoneIndexToMeshBoneIndexMap = InReferencePose.GetLODBoneIndexToMeshBoneIndexMap(0);
+
+		checkf(LODBoneIndexToSkeletonBoneIndexMap.Num() == RefPoseBoneNum, TEXT("Buffer mismatch: %d:%d"), LODBoneIndexToSkeletonBoneIndexMap.Num(), RefPoseBoneNum);
+		checkf(MeshBoneIndexToLODBoneIndexMap.Num() == RefPoseBoneNum, TEXT("Buffer mismatch: %d:%d"), MeshBoneIndexToLODBoneIndexMap.Num(), RefPoseBoneNum);
+		checkf(BoneNameToLODBoneIndexMap.Num() == RefPoseBoneNum, TEXT("Buffer mismatch: %d:%d"), BoneNameToLODBoneIndexMap.Num(), RefPoseBoneNum);
+		checkf(LODBoneIndexToMeshBoneIndexMap.Num() == RefPoseBoneNum, TEXT("Buffer mismatch: %d:%d"), LODBoneIndexToMeshBoneIndexMap.Num(), RefPoseBoneNum);
 
 		// Build mirror map for all the mesh bones in reference pose (LOD0).
 		for (FBoneIndexType MeshBoneIndex = 0; MeshBoneIndex < OutMeshBoneIndexToMirroredMeshBoneIndexMap.Num(); ++MeshBoneIndex)
 		{
 			if (OutMeshBoneIndexToMirroredMeshBoneIndexMap[MeshBoneIndex] == FBoneIndexType(INDEX_NONE))
 			{
-				// Find the candidate mirror partner for this bone (falling back to mirroring to self)
-				const FName SourceBoneName = ReferenceSkeleton.GetBoneName(MeshBoneIndex);
-				int32 MirrorMeshBoneIndex = INDEX_NONE;
+				// Get index for LOD sorted map
+				const FBoneIndexType LODBoneIndex = MeshBoneIndexToLODBoneIndexMap[MeshBoneIndex];
+				checkf(LODBoneIndex != FBoneIndexType(INDEX_NONE), TEXT("Failed to find LOD Bone Index counterpart for Mesh Bone Index"))
+
+				// We need to convert from our sorted LOD bone map to the skeleton bone map, otherwise the mapping is broken.
+				const FBoneIndexType SkeletonBoneIndex = LODBoneIndexToSkeletonBoneIndexMap[LODBoneIndex];
+				
+				// Find the candidate mirror partner for this bone
+				const FName SourceBoneName = ReferenceSkeleton.GetBoneName(SkeletonBoneIndex);
+				FBoneIndexType MirrorMeshBoneIndex = INDEX_NONE;
 				
 				FName* MirroredBoneName = NameToMirrorNameBoneMap.Find(SourceBoneName);
-				if (!SourceBoneName.IsNone() && MirroredBoneName)
+				if (!SourceBoneName.IsNone() && MirroredBoneName && BoneNameToLODBoneIndexMap.Contains(*MirroredBoneName))
 				{
-					MirrorMeshBoneIndex = ReferenceSkeleton.FindBoneIndex(*MirroredBoneName);
+					int32 MirrorLODBoneIndex = BoneNameToLODBoneIndexMap[*MirroredBoneName];
+
+					if (MirrorLODBoneIndex != FBoneIndexType(INDEX_NONE))
+					{
+						MirrorMeshBoneIndex = LODBoneIndexToMeshBoneIndexMap[MirrorLODBoneIndex];
+					}
 				}
 				
 				OutMeshBoneIndexToMirroredMeshBoneIndexMap[MeshBoneIndex] = MirrorMeshBoneIndex;
 				
 				// Map candidate mirror partner to current bone.
 				// @todo: What happens we have conflicts on mirror data table setup. Should we have a conflict policy to resolve this? For example, two bones map to another.
-				if (MirrorMeshBoneIndex != INDEX_NONE)
+				if (MirrorMeshBoneIndex != FBoneIndexType(INDEX_NONE))
 				{
 					OutMeshBoneIndexToMirroredMeshBoneIndexMap[MirrorMeshBoneIndex] = MeshBoneIndex;
 				}
@@ -163,8 +182,8 @@ namespace UE::UAF
 		
 		// Get mirrored bones mapping.
 		TArray<FBoneIndexType, FAnimStackAllocator> MeshBoneIndexToMirroredMeshBoneIndexMap;
-		TArray<FQuat, FAnimStackAllocator> ComponentSpaceReferencePoseRotations;
-		TArray<FQuat, FAnimStackAllocator> ComponentSpaceReferenceRotationCorrections;
+		TArray<FQuat, FAnimStackAllocator> MeshSpaceReferencePoseRotations;
+		TArray<FQuat, FAnimStackAllocator> MeshSpaceReferenceRotationCorrections;
 
 		const FReferencePose& ReferencePose = InOutLODPose.GetRefPose();
 		
@@ -172,13 +191,13 @@ namespace UE::UAF
 		MeshBoneIndexToMirroredMeshBoneIndexMap.SetNumUninitialized(MirrorBoneIndicesNum);
 
 		const int32 BindPoseMirrorDataNum = GetNumOfBonesForMirrorData(InOutLODPose.GetRefPose());
-		ComponentSpaceReferencePoseRotations.SetNumUninitialized(BindPoseMirrorDataNum);
-		ComponentSpaceReferenceRotationCorrections.SetNumUninitialized(BindPoseMirrorDataNum);
+		MeshSpaceReferencePoseRotations.SetNumUninitialized(BindPoseMirrorDataNum);
+		MeshSpaceReferenceRotationCorrections.SetNumUninitialized(BindPoseMirrorDataNum);
 		
 		BuildMeshBoneIndexMirrorMap(ReferencePose, InMirrorDataTable, MeshBoneIndexToMirroredMeshBoneIndexMap);
-		BuildReferencePoseMirrorData(ReferencePose, InMirrorDataTable.MirrorAxis, MeshBoneIndexToMirroredMeshBoneIndexMap, ComponentSpaceReferencePoseRotations, ComponentSpaceReferenceRotationCorrections);
+		BuildReferencePoseMirrorData(ReferencePose, InMirrorDataTable.MirrorAxis, MeshBoneIndexToMirroredMeshBoneIndexMap, MeshSpaceReferencePoseRotations, MeshSpaceReferenceRotationCorrections);
 		
-		MirrorPose(InOutLODPose, InMirrorDataTable.MirrorAxis, MeshBoneIndexToMirroredMeshBoneIndexMap, ComponentSpaceReferencePoseRotations, ComponentSpaceReferenceRotationCorrections);
+		MirrorPose(InOutLODPose, InMirrorDataTable.MirrorAxis, MeshBoneIndexToMirroredMeshBoneIndexMap, MeshSpaceReferencePoseRotations, MeshSpaceReferenceRotationCorrections);
 	}
 	
 	void MirrorPose(
@@ -188,8 +207,6 @@ namespace UE::UAF
 		TConstArrayView<FQuat> InRefPoseMeshSpaceRotations,
 		TConstArrayView<FQuat> InRefPoseMeshSpaceRotationCorrections)
 	{
-		// @todo: check mirror map size
-		
 		const int32 SourceLODBoneNum = InOutLODPose.LocalTransformsView.Num();
 		
 		const int32 RefPoseBoneNum = InOutLODPose.GetRefPose().ReferenceLocalTransforms.Num();
@@ -215,6 +232,12 @@ namespace UE::UAF
 		
 		const TArrayView<const FBoneIndexType>& LODBoneIndexToParentLODBoneIndexMap = InOutLODPose.GetLODBoneIndexToParentLODBoneIndexMap();
 		checkf(LODBoneIndexToParentLODBoneIndexMap.Num() == InOutLODPose.GetNumBones(), TEXT("Buffer mismatch: %d:%d"), LODBoneIndexToParentLODBoneIndexMap.Num(), InOutLODPose.GetNumBones());
+
+		const TArrayView<const FBoneIndexType>& LODBoneIndexToMeshBoneIndexMap = InOutLODPose.GetLODBoneIndexToMeshBoneIndexMap();
+		checkf(LODBoneIndexToMeshBoneIndexMap.Num() == InOutLODPose.GetNumBones(), TEXT("Buffer mismatch: %d:%d"), LODBoneIndexToMeshBoneIndexMap.Num(), InOutLODPose.GetNumBones());
+
+		const TArrayView<const FBoneIndexType>& MeshBoneIndexToLODBoneIndexMap = InOutLODPose.GetMeshBoneIndexToLODBoneIndexMap();
+		checkf(MeshBoneIndexToLODBoneIndexMap.Num() == InRefPoseMeshSpaceRotationCorrections.Num(), TEXT("Buffer mismatch: %d:%d"), MeshBoneIndexToLODBoneIndexMap.Num(), InRefPoseMeshSpaceRotationCorrections.Num())
 		
 		// Mirroring is authored in object space and as such we must transform the local space transforms in object space in order
 		// to apply the object space mirroring axis. To facilitate this, we use object space transforms for the bind pose which can be cached.
@@ -251,8 +274,7 @@ namespace UE::UAF
 
 			// Mirror the rotation component: Rotate into the space of the mirror plane, mirror across the plane, apply corrective rotation to align result with target space's rest orientation, 
 			// then rotate into the space of its new parent
-
-			
+				
 			FQuat Q = SourceTransform.GetRotation();
 			Q = SourceParentRefRotation * Q;                     // to mesh space (component space)
 			Q = FAnimationRuntime::MirrorQuat(Q, InMirrorAxis);  // reflect accros plane
@@ -273,57 +295,96 @@ namespace UE::UAF
 		// Mirror the root bone.
 		{
 			// @todo: Can the root bone be other than 0? LODPose.GetRefPose().RootBoneIndex? Can we have more than one root bone? 
-			FBoneIndexType RootBoneIndex = 0;
-			FBoneIndexType MirrorRootBoneIndex = InMeshBoneIndexToMirroredMeshBoneIndexMap[RootBoneIndex];
+			const FBoneIndexType RootMeshBoneIndex = 0;
+			const FBoneIndexType RootLODBoneIndex =  InMeshBoneIndexToMirroredMeshBoneIndexMap[RootMeshBoneIndex];
+			const FBoneIndexType MirrorRootMeshBoneIndex = InMeshBoneIndexToMirroredMeshBoneIndexMap[RootMeshBoneIndex];
 
-			if (MirrorRootBoneIndex != FBoneIndexType(INDEX_NONE))
+			if (MirrorRootMeshBoneIndex != FBoneIndexType(INDEX_NONE) || RootLODBoneIndex != FBoneIndexType(INDEX_NONE))
 			{
-				const FBoneIndexType TargetParentIndex = LODBoneIndexToParentLODBoneIndexMap[RootBoneIndex];
-				const FBoneIndexType SourceParentIndex = LODBoneIndexToParentLODBoneIndexMap[MirrorRootBoneIndex];
+				const FBoneIndexType MirroredRootLODBoneIndex = MeshBoneIndexToLODBoneIndexMap[MirrorRootMeshBoneIndex];
 
-				if (TargetParentIndex != FBoneIndexType(INDEX_NONE))
+				if (MirroredRootLODBoneIndex == FBoneIndexType(INDEX_NONE))
 				{
-					UE_LOGFMT(LogAnimation, Error, "UAF::MirrorPose - Found parent bone index ({0}) for root bone index ({1}). This is invalid.", TargetParentIndex, RootBoneIndex);
+					UE_LOGFMT(LogAnimation, Error, "UAF::MirrorPose - Failed to find LOD Bone index ({0}) for Mesh Bone Index ({1}) for a mirrored root bone. This is invalid.", MirroredRootLODBoneIndex, MirrorRootMeshBoneIndex);
 					return;
 				}
 				
-				if (TargetParentIndex == FBoneIndexType(INDEX_NONE) && SourceParentIndex != FBoneIndexType(INDEX_NONE))
+				const FBoneIndexType TargetParentLODBoneIndex = LODBoneIndexToParentLODBoneIndexMap[RootMeshBoneIndex];
+				const FBoneIndexType SourceParentLODBoneIndex = LODBoneIndexToParentLODBoneIndexMap[MirroredRootLODBoneIndex];
+
+				if (TargetParentLODBoneIndex != FBoneIndexType(INDEX_NONE))
+				{
+					UE_LOGFMT(LogAnimation, Error, "UAF::MirrorPose - Found parent bone index ({0}) for root bone index ({1}). This is invalid.", TargetParentLODBoneIndex, RootMeshBoneIndex);
+					return;
+				}
+
+				if (SourceParentLODBoneIndex != FBoneIndexType(INDEX_NONE))
 				{
-					UE_LOGFMT(LogAnimation, Error, "UAF::MirrorPose - Mapping root bone to non-root bone is not supported ({0} -> {1})", RootBoneIndex, MirrorRootBoneIndex);
+					UE_LOGFMT(LogAnimation, Error, "UAF::MirrorPose - Found parent bone index ({0}) for mirrored root bone index ({1}). This is invalid.", TargetParentLODBoneIndex, RootMeshBoneIndex);
 					return;
 				}
 				
-				InOutLODPose.LocalTransformsView[RootBoneIndex] = MirrorTransform(InOutLODPose.LocalTransformsView[MirrorRootBoneIndex], SourceParentIndex, MirrorRootBoneIndex, TargetParentIndex, RootBoneIndex);
+				if (TargetParentLODBoneIndex == FBoneIndexType(INDEX_NONE) && SourceParentLODBoneIndex != FBoneIndexType(INDEX_NONE))
+				{
+					UE_LOGFMT(LogAnimation, Error, "UAF::MirrorPose - Mapping root bone to non-root bone is not supported ({0} -> {1})", RootMeshBoneIndex, MirroredRootLODBoneIndex);
+					return;
+				}
+
+				// At this point the parent bone of the root should be -1 so we're good to just set it as is.
+				const FBoneIndexType TargetParentMeshBoneIndex = TargetParentLODBoneIndex;
+				const FBoneIndexType SourceParentMeshBoneIndex = SourceParentLODBoneIndex;
+				
+				InOutLODPose.LocalTransformsView[RootLODBoneIndex] = MirrorTransform(InOutLODPose.LocalTransformsView[RootLODBoneIndex], SourceParentMeshBoneIndex, RootMeshBoneIndex, TargetParentMeshBoneIndex, RootMeshBoneIndex);
 			}
 		}
 		
 		// Mirror the non-root bones.
 		for (FBoneIndexType LODBoneIndex = 1; LODBoneIndex < SourceLODBoneNum; ++LODBoneIndex)
 		{
+			// Skip invalid mesh bone index (should never happen).
+			const FBoneIndexType MeshBoneIndex = LODBoneIndex != FBoneIndexType(INDEX_NONE) ? LODBoneIndexToMeshBoneIndexMap[LODBoneIndex] : FBoneIndexType(INDEX_NONE);
+			if (MeshBoneIndex == FBoneIndexType(INDEX_NONE))
+			{
+				continue;
+			}
+				
 			// @note: We can safely use the LODBoneIndex to get our mirrored name since the map is in LOD0 which will always contain the index.
-			const FBoneIndexType MirroredLODBoneIndex = InMeshBoneIndexToMirroredMeshBoneIndexMap[LODBoneIndex];
+			const FBoneIndexType MirroredMeshBoneIndex = InMeshBoneIndexToMirroredMeshBoneIndexMap[MeshBoneIndex];
+
+			// Skip invalid mirrors.
+			const bool bIsMirrorInvalid = MirroredMeshBoneIndex == FBoneIndexType(INDEX_NONE);
+			if (bIsMirrorInvalid)
+			{
+				continue;	
+			}
+			
+			const FBoneIndexType MirroredLODBoneIndex = MeshBoneIndexToLODBoneIndexMap[MirroredMeshBoneIndex];
 			const FBoneIndexType ParentLODBoneIndex = LODBoneIndexToParentLODBoneIndexMap[LODBoneIndex];
+			const FBoneIndexType ParentMeshBoneIndex = ParentLODBoneIndex != FBoneIndexType(INDEX_NONE) ? LODBoneIndexToMeshBoneIndexMap[ParentLODBoneIndex] : FBoneIndexType(INDEX_NONE);
 			
 			// Self-mirror (mirror in place relative to the same parent).
-			if (LODBoneIndex == MirroredLODBoneIndex)
+			const bool bIsMirroredBoneMappedToSelf = LODBoneIndex == MirroredLODBoneIndex;
+			if (bIsMirroredBoneMappedToSelf)
 			{
-				InOutLODPose.LocalTransformsView[LODBoneIndex] = MirrorTransform(InOutLODPose.LocalTransformsView[LODBoneIndex], ParentLODBoneIndex, LODBoneIndex, ParentLODBoneIndex, LODBoneIndex);
+				InOutLODPose.LocalTransformsView[LODBoneIndex] = MirrorTransform(InOutLODPose.LocalTransformsView[LODBoneIndex], ParentMeshBoneIndex, MeshBoneIndex, ParentMeshBoneIndex, MeshBoneIndex);
 				continue;
 			}
 
 			// Skip invalid or already-processed mirror pairs.
-			if (MirroredLODBoneIndex == FBoneIndexType(INDEX_NONE) || MirroredLODBoneIndex < LODBoneIndex)
+			const bool bIsAlreadyProcessed = MirroredLODBoneIndex < LODBoneIndex;
+			if (bIsAlreadyProcessed)
 			{
 				continue;
 			}
 			
-			const FBoneIndexType MirroredParentLODIndex = LODBoneIndexToParentLODBoneIndexMap[MirroredLODBoneIndex];
-
+			const FBoneIndexType MirroredParentLODBoneIndex = LODBoneIndexToParentLODBoneIndexMap[MirroredLODBoneIndex];
+			const FBoneIndexType MirroredParentMeshBoneIndex = LODBoneIndexToMeshBoneIndexMap[MirroredParentLODBoneIndex];
+			
 			const FTransform OriginalTransformAtBoneIndex = InOutLODPose.LocalTransformsView[LODBoneIndex];
 			const FTransform OriginalTransformAtMirroredIndex = InOutLODPose.LocalTransformsView[MirroredLODBoneIndex];
 			
-			const FTransform NewTransformAtBoneIndex = MirrorTransform(OriginalTransformAtMirroredIndex, MirroredParentLODIndex, MirroredLODBoneIndex, ParentLODBoneIndex, LODBoneIndex);
-			const FTransform NewTransformAtMirroredIndex = MirrorTransform(OriginalTransformAtBoneIndex, ParentLODBoneIndex, LODBoneIndex, MirroredParentLODIndex, MirroredLODBoneIndex);
+			const FTransform NewTransformAtBoneIndex = MirrorTransform(OriginalTransformAtMirroredIndex, MirroredParentMeshBoneIndex, MirroredMeshBoneIndex, ParentMeshBoneIndex, MeshBoneIndex);
+			const FTransform NewTransformAtMirroredIndex = MirrorTransform(OriginalTransformAtBoneIndex, ParentMeshBoneIndex, MeshBoneIndex, MirroredParentMeshBoneIndex, MirroredMeshBoneIndex);
 			
 			InOutLODPose.LocalTransformsView[LODBoneIndex] = NewTransformAtBoneIndex;
 			InOutLODPose.LocalTransformsView[MirroredLODBoneIndex] = NewTransformAtMirroredIndex;
