Removed error check that will never be reached in mirroring

#rnx
#Synced-CL 48298300

diff --git a/Engine/Plugins/Experimental/UAF/UAFMirroring/Source/UAFMirroring/Private/Mirroring.cpp b/Engine/Plugins/Experimental/UAF/UAFMirroring/Source/UAFMirroring/Private/Mirroring.cpp
--- a/Engine/Plugins/Experimental/UAF/UAFMirroring/Source/UAFMirroring/Private/Mirroring.cpp
+++ b/Engine/Plugins/Experimental/UAF/UAFMirroring/Source/UAFMirroring/Private/Mirroring.cpp
@@ -323,12 +323,6 @@ namespace UE::UAF
 					UE_LOGFMT(LogAnimation, Error, "UAF::MirrorPose - Found parent bone index ({0}) for mirrored root bone index ({1}). This is invalid.", TargetParentLODBoneIndex, RootMeshBoneIndex);
 					return;
 				}
-				
-				if (TargetParentLODBoneIndex == FBoneIndexType(INDEX_NONE) && SourceParentLODBoneIndex != FBoneIndexType(INDEX_NONE))
-				{
-					UE_LOGFMT(LogAnimation, Error, "UAF::MirrorPose - Mapping root bone to non-root bone is not supported ({0} -> {1})", RootMeshBoneIndex, MirroredRootLODBoneIndex);
-					return;
-				}
 
 				// At this point the parent bone of the root should be -1 so we're good to just set it as is.
 				const FBoneIndexType TargetParentMeshBoneIndex = TargetParentLODBoneIndex;
