UE-354797 Fix a  UAF crash due to out of bounds index using the wrong reference skeleton in mirror name lookup
#Synced-CL 48694509

diff --git a/Engine/Plugins/Experimental/UAF/UAFMirroring/Source/UAFMirroring/Private/Mirroring.cpp b/Engine/Plugins/Experimental/UAF/UAFMirroring/Source/UAFMirroring/Private/Mirroring.cpp
--- a/Engine/Plugins/Experimental/UAF/UAFMirroring/Source/UAFMirroring/Private/Mirroring.cpp
+++ b/Engine/Plugins/Experimental/UAF/UAFMirroring/Source/UAFMirroring/Private/Mirroring.cpp
@@ -107,6 +107,9 @@ namespace UE::UAF
 		checkf(SkeletalMesh != nullptr, TEXT("UAF::FillMirrorBoneIndices - SkeletalMesh is null"));
 		const FReferenceSkeleton& ReferenceSkeleton = SkeletalMesh->GetRefSkeleton();
 		
+		const USkeleton* Skeleton = SkeletalMesh->GetSkeleton(); 
+		const FReferenceSkeleton& FullLODSkeleton = Skeleton->GetReferenceSkeleton(); 
+
 		const TArrayView<const FBoneIndexType> LODBoneIndexToSkeletonBoneIndexMap = InReferencePose.GetLODBoneIndexToSkeletonBoneIndexMap(0);
 		const TArrayView<const FBoneIndexType> MeshBoneIndexToLODBoneIndexMap = InReferencePose.GetMeshBoneIndexToLODBoneIndexMap();
 		const TMap<FName, FBoneIndexType> BoneNameToLODBoneIndexMap = InReferencePose.GetBoneNameToLODBoneIndexMap();
@@ -129,8 +132,9 @@ namespace UE::UAF
 				// We need to convert from our sorted LOD bone map to the skeleton bone map, otherwise the mapping is broken.
 				const FBoneIndexType SkeletonBoneIndex = LODBoneIndexToSkeletonBoneIndexMap[LODBoneIndex];
 				
+				// since this is a skeleton index we need to use the skeleton to find the name 
 				// Find the candidate mirror partner for this bone
-				const FName SourceBoneName = ReferenceSkeleton.GetBoneName(SkeletonBoneIndex);
+				const FName SourceBoneName = FullLODSkeleton.GetBoneName(SkeletonBoneIndex);
 				FBoneIndexType MirrorMeshBoneIndex = INDEX_NONE;
 				
 				FName* MirroredBoneName = NameToMirrorNameBoneMap.Find(SourceBoneName);
diff --git a/Engine/Plugins/Experimental/UAF/UAFMirroring/Source/UAFMirroring/Private/MirroringTrait.cpp b/Engine/Plugins/Experimental/UAF/UAFMirroring/Source/UAFMirroring/Private/MirroringTrait.cpp
--- a/Engine/Plugins/Experimental/UAF/UAFMirroring/Source/UAFMirroring/Private/MirroringTrait.cpp
+++ b/Engine/Plugins/Experimental/UAF/UAFMirroring/Source/UAFMirroring/Private/MirroringTrait.cpp
@@ -18,10 +18,11 @@ namespace UE::UAF
 	////////////////////////////////////////////////////////////////////////////////////////////////////////////////////////
 	// FMirroringTraitCache
 
+
 	bool FMirroringTraitCache::AreMirrorMapsValid(const UE::UAF::FReferencePose& InReferencePose, const TWeakObjectPtr<const UMirrorDataTable>& InMirrorTable, bool bInShouldMirrorBones, bool bInShouldMirrorAttributes) const
 	{
 		// Ensure assets are the same.
-		
+
 		if (InMirrorTable != MirrorTable)
 		{
 			return false;
@@ -33,7 +34,7 @@ namespace UE::UAF
 		}
 
 		// Ensure arrays match mesh bone count.
-		
+
 		const int32 ExpectedNumBones = UE::UAF::GetNumOfBonesForMirrorData(InReferencePose);
 
 		// Default to true if we are not mirroring bones since these arrays are only used for that channel.
@@ -41,7 +42,7 @@ namespace UE::UAF
 
 		// Default to true if we are not mirroring attributes since these arrays are only used for that channel.
 		const bool bDoesCompactPoseMirrorMapHaveExpectedNumBones = (ExpectedNumBones == CompactPoseBoneIndexToMirroredCompactPoseBoneIndexMap.Num() || !bInShouldMirrorAttributes);
-		
+
 		const bool bAreMirrorMapsEquallySized = bDoesMeshMirrorMapHaveExpectedNumBones && bDoesCompactPoseMirrorMapHaveExpectedNumBones;
 		if (!bAreMirrorMapsEquallySized)
 		{
@@ -54,7 +55,7 @@ namespace UE::UAF
 		{
 			return false;
 		}
-		
+
 		return true;
 	}
 	
@@ -286,14 +287,6 @@ void FAnimNextEvaluationMirroringTask::Execute(UE::UAF::FEvaluationVM& VM) const
 		FAnimationRuntime::MirrorCurves(Keyframe->Curves, *Setup.MirrorDataTable);
 	}
 
-	if (bShouldMirrorAttributes)
-	{
-		// @todo: Make UAF version of MirrorAttributes() or address internal heap allocations for scratch array.
-		UE::Anim::Attributes::MirrorAttributes(
-			Keyframe->Attributes,
-			*Setup.MirrorDataTable,
-			Cache->CompactPoseBoneIndexToMirroredCompactPoseBoneIndexMap);
-	}
 
 	VM.PushValue(UE::UAF::KEYFRAME_STACK_NAME, MoveTemp(Keyframe));
 }
@@ -315,7 +308,6 @@ void FAnimNextEvaluationMirroringTask::EnsureCache(const UE::UAF::FEvaluationVM&
 			if (!bAreMirrorMapsValid)
 			{
 				Cache->MeshBoneIndexToMirroredMeshBoneIndexMap.SetNumUninitialized(NumBonesForLOD0);
-				
 				UE::UAF::BuildMeshBoneIndexMirrorMap(InReferencePose, *Setup.MirrorDataTable, Cache->MeshBoneIndexToMirroredMeshBoneIndexMap);
 			}
 			
@@ -332,31 +324,34 @@ void FAnimNextEvaluationMirroringTask::EnsureCache(const UE::UAF::FEvaluationVM&
 					Cache->MeshSpaceReferenceRotationCorrections
 				);
 			}
-		}
-		
-		if (bShouldMirrorAttributes)
-		{
-			if (!bAreMirrorMapsValid)
+
+			if (bShouldMirrorAttributes)
 			{
-				// We only care about the mesh bones (LOD0).
-				TArray<FBoneIndexType> RequiredBoneIndexArray;
-				RequiredBoneIndexArray.AddUninitialized(NumBonesForLOD0);
-				
-				for (int32 BoneIndex = 0; BoneIndex < RequiredBoneIndexArray.Num(); ++BoneIndex)
+				if (!bAreMirrorMapsValid)
 				{
-					RequiredBoneIndexArray[BoneIndex] = InReferencePose.GetSkeletonBoneIndexFromLODBoneIndex(BoneIndex);
-				}
+					// We only care about the mesh bones (LOD0).
+					TArray<FBoneIndexType> RequiredBoneIndexArray;
+					RequiredBoneIndexArray.AddUninitialized(NumBonesForLOD0);
 
-				// Compute mirror map for compact pose.
-				
-				FBoneContainer BoneContainer;
-				BoneContainer.InitializeTo(RequiredBoneIndexArray, UE::Anim::FCurveFilterSettings(UE::Anim::ECurveFilterMode::DisallowAll), *InReferencePose.SkeletalMesh);
-				
-				TCustomBoneIndexArray<FSkeletonPoseBoneIndex, FSkeletonPoseBoneIndex> MirrorBoneIndexes;
-				Setup.MirrorDataTable->FillMirrorBoneIndexes(BoneContainer.GetSkeletonAsset(), MirrorBoneIndexes);
-				Setup.MirrorDataTable->FillCompactPoseMirrorBones(BoneContainer, MirrorBoneIndexes, Cache->CompactPoseBoneIndexToMirroredCompactPoseBoneIndexMap);
+					for (int32 BoneIndex = 0; BoneIndex < RequiredBoneIndexArray.Num(); ++BoneIndex)
+					{
+						RequiredBoneIndexArray[BoneIndex] = InReferencePose.GetSkeletonBoneIndexFromLODBoneIndex(BoneIndex);
+					}
+
+					// Compute mirror map for compact pose.
+
+					FBoneContainer BoneContainer;
+					const USkeleton* Skeleton = InReferencePose.SkeletalMesh->GetSkeleton();
+					BoneContainer.InitializeTo(RequiredBoneIndexArray, UE::Anim::FCurveFilterSettings(UE::Anim::ECurveFilterMode::DisallowAll), *Skeleton);
+
+					TCustomBoneIndexArray<FSkeletonPoseBoneIndex, FSkeletonPoseBoneIndex> MirrorBoneIndexes;
+					Setup.MirrorDataTable->FillMirrorBoneIndexes(BoneContainer.GetSkeletonAsset(), MirrorBoneIndexes);
+					Setup.MirrorDataTable->FillCompactPoseMirrorBones(BoneContainer, MirrorBoneIndexes, Cache->CompactPoseBoneIndexToMirroredCompactPoseBoneIndexMap);
+				}
 			}
 		}
+		
+	
 	}
 	else
 	{
