Anim Mixer: Fix mirroring support of skeletal animation sections.

As the first step of supporting mirroring in the mixer, this gets it working in a backwards-compatible way with skeletal animation sections, where the mirror data table is
specified as a property on the Skeletal Animation Section.

We do this by using the existing Mirror Evaluation Task, which is refactored out of the trait file, as well as moving the runtime instance data responsible for mirroring into the task itself rather than leaving it in the trait.

We then append an additional mirroring evaluation task in the case we want to mirror. This strategy allows us to use mirroring on any kind of animation section going forward. That functionality will come in a later changelist.

#Synced-CL 49622961

diff --git a/Engine/Plugins/Experimental/MovieSceneAnimMixer/MovieSceneAnimMixer.uplugin b/Engine/Plugins/Experimental/MovieSceneAnimMixer/MovieSceneAnimMixer.uplugin
--- a/Engine/Plugins/Experimental/MovieSceneAnimMixer/MovieSceneAnimMixer.uplugin
+++ b/Engine/Plugins/Experimental/MovieSceneAnimMixer/MovieSceneAnimMixer.uplugin
@@ -24,6 +24,10 @@
 		{
 			"Name": "UAFAnimGraph",
 			"Enabled": true
+		},
+		{
+			"Name": "UAFMirroring",
+			"Enabled": true
 		}
 	],
 	"Modules" :
diff --git a/Engine/Plugins/Experimental/MovieSceneAnimMixer/Source/MovieSceneAnimMixer/MovieSceneAnimMixer.Build.cs b/Engine/Plugins/Experimental/MovieSceneAnimMixer/Source/MovieSceneAnimMixer/MovieSceneAnimMixer.Build.cs
--- a/Engine/Plugins/Experimental/MovieSceneAnimMixer/Source/MovieSceneAnimMixer/MovieSceneAnimMixer.Build.cs
+++ b/Engine/Plugins/Experimental/MovieSceneAnimMixer/Source/MovieSceneAnimMixer/MovieSceneAnimMixer.Build.cs
@@ -19,7 +19,8 @@ namespace UnrealBuildTool.Rules
 					"AnimGraphRuntime",
 					"UAF",
 					"UAFAnimGraph",
-					"SlateCore"
+					"SlateCore",
+					"UAFMirroring"
 				}
 			);
 
diff --git a/Engine/Plugins/Experimental/MovieSceneAnimMixer/Source/MovieSceneAnimMixer/Private/AnimMixerComponentTypes.cpp b/Engine/Plugins/Experimental/MovieSceneAnimMixer/Source/MovieSceneAnimMixer/Private/AnimMixerComponentTypes.cpp
--- a/Engine/Plugins/Experimental/MovieSceneAnimMixer/Source/MovieSceneAnimMixer/Private/AnimMixerComponentTypes.cpp
+++ b/Engine/Plugins/Experimental/MovieSceneAnimMixer/Source/MovieSceneAnimMixer/Private/AnimMixerComponentTypes.cpp
@@ -48,6 +48,7 @@ namespace UE::MovieScene
 		ComponentRegistry->NewComponentType(&MixerEntry, TEXT("MixerEntry"));
 		ComponentRegistry->NewComponentType(&RootDestination, TEXT("Root Destination"));
 		ComponentRegistry->NewComponentType(&MixerLayer, TEXT("Mixer Layer"));
+		ComponentRegistry->NewComponentType(&MirrorTable, TEXT("Mirror Table"));
 
 
 		Tags.RequiresBlending = ComponentRegistry->NewTag(TEXT("Requires Blending"));
@@ -62,6 +63,8 @@ namespace UE::MovieScene
 		ComponentRegistry->Factories.DuplicateChildComponent(RootDestination);
 		ComponentRegistry->Factories.DuplicateChildComponent(MixerLayer);
 
+		ComponentRegistry->Factories.DuplicateChildComponent(MirrorTable);
+		
 		ComponentRegistry->Factories.DefineChildComponent(RootDestination, MixerRootMotion);
 		ComponentRegistry->Factories.DefineChildComponent(Task, MixerEntry);
 
diff --git a/Engine/Plugins/Experimental/MovieSceneAnimMixer/Source/MovieSceneAnimMixer/Private/Systems/MovieSceneAnimMixerSystem.cpp b/Engine/Plugins/Experimental/MovieSceneAnimMixer/Source/MovieSceneAnimMixer/Private/Systems/MovieSceneAnimMixerSystem.cpp
--- a/Engine/Plugins/Experimental/MovieSceneAnimMixer/Source/MovieSceneAnimMixer/Private/Systems/MovieSceneAnimMixerSystem.cpp
+++ b/Engine/Plugins/Experimental/MovieSceneAnimMixer/Source/MovieSceneAnimMixer/Private/Systems/MovieSceneAnimMixerSystem.cpp
@@ -45,6 +45,7 @@
 #include "Systems/MovieSceneRootMotionSystem.h"
 #include "AnimSubsystem_SequencerMixer.h"
 #include "Sections/MovieSceneSkeletalAnimationSection.h"
+#include "MirroringTask.h"
 
 #include UE_INLINE_GENERATED_CPP_BY_NAME(MovieSceneAnimMixerSystem)
 
@@ -228,6 +229,15 @@ namespace UE::MovieScene
 					// Evaluate the pose itself
 					Mixer->EvaluationProgram->AppendTaskPtr(MixerEntry->EvalTask);
 
+					// Mirror if specified
+					if (MixerEntry->MirrorTable)
+					{
+						// Add a mirroring task
+						TSharedPtr<FAnimNextEvaluationMirroringTask> MirrorTask = MakeShared<FAnimNextEvaluationMirroringTask>();
+						MirrorTask->Setup.MirrorDataTable = MixerEntry->MirrorTable;
+						Mixer->EvaluationProgram->AppendTaskPtr(MirrorTask);
+					}
+
 					// Handle root space conversions and manipulations
 					{
 						FTransform TransformOrigin  = FTransform::Identity;
@@ -557,6 +567,7 @@ UMovieSceneAnimMixerSystem::UMovieSceneAnimMixerSystem(const FObjectInitializer&
 		DefineComponentConsumer(GetClass(), AnimMixerComponents->Task);
 		DefineComponentConsumer(GetClass(), AnimMixerComponents->Target);
 		DefineComponentConsumer(GetClass(), AnimMixerComponents->Priority);
+		DefineComponentConsumer(GetClass(), AnimMixerComponents->MirrorTable);
 		DefineComponentProducer(GetClass(), AnimMixerComponents->MixerTask);
 	}
 }
@@ -689,6 +700,7 @@ void UMovieSceneAnimMixerSystem::OnRun(FSystemTaskPrerequisites& InPrerequisites
 		.Read(AnimMixerComponents->Priority)
 		.ReadOptional(TrackComponents->SkeletalAnimation)
 		.ReadOptional(AnimMixerComponents->RootMotionSettings)
+		.ReadOptional(AnimMixerComponents->MirrorTable)
 		.FilterAll({ BuiltInComponents->Tags.NeedsLink})
 		.Iterate_PerEntity( &Linker->EntityManager,
 		[this, &NewMixers, BuiltInComponents, AnimMixerComponents]
@@ -701,7 +713,8 @@ void UMovieSceneAnimMixerSystem::OnRun(FSystemTaskPrerequisites& InPrerequisites
 			TSharedPtr<FAnimNextEvaluationTask> Task,
 			int32 Priority,
 			const FMovieSceneSkeletalAnimationComponentData* SkelAnimData,
-			const FMovieSceneRootMotionSettings* RootMotionSettings)
+			const FMovieSceneRootMotionSettings* RootMotionSettings,
+			const TObjectPtr<UMirrorDataTable>* OptMirrorDataTable)
 			{
 				// For new entities, we 'resolve' the animation target so if 'Automatic' is picked we choose the right one automatically.
 				Target = ResolveAnimationTarget(BoundObjectKey, Target, SkelAnimData);
@@ -767,10 +780,10 @@ void UMovieSceneAnimMixerSystem::OnRun(FSystemTaskPrerequisites& InPrerequisites
 				{
 					Entry->RootMotionLifetimeReference = nullptr;
 				}
-
 				const FComponentMask& Type = Linker->EntityManager.GetEntityType(EntityID);
 				Entry->bAdditive = Type.Contains(BuiltInComponents->Tags.AdditiveAnimation);
 				Entry->bRequiresBlend = Type.Contains(AnimMixerComponents->Tags.RequiresBlending);
+				Entry->MirrorTable = OptMirrorDataTable ? *OptMirrorDataTable : nullptr;
 
 				TSharedPtr<FMovieSceneAnimMixer> ExistingMixer = Entry->WeakMixer.Pin();
 				if (ExistingMixer)
diff --git a/Engine/Plugins/Experimental/MovieSceneAnimMixer/Source/MovieSceneAnimMixer/Private/Systems/MovieSceneMixedSkeletalAnimationSystem.cpp b/Engine/Plugins/Experimental/MovieSceneAnimMixer/Source/MovieSceneAnimMixer/Private/Systems/MovieSceneMixedSkeletalAnimationSystem.cpp
--- a/Engine/Plugins/Experimental/MovieSceneAnimMixer/Source/MovieSceneAnimMixer/Private/Systems/MovieSceneMixedSkeletalAnimationSystem.cpp
+++ b/Engine/Plugins/Experimental/MovieSceneAnimMixer/Source/MovieSceneAnimMixer/Private/Systems/MovieSceneMixedSkeletalAnimationSystem.cpp
@@ -89,6 +89,7 @@ struct FSkeletalAnimMixerMutation : IMovieSceneEntityMutation
 			AnimMixerComponents->Task,
 			AnimMixerComponents->MixerEntry,
 			AnimMixerComponents->RootMotionSettings,
+			AnimMixerComponents->MirrorTable,
 			AnimMixerComponents->Tags.RequiresBlending
 		});
 	}
@@ -122,10 +123,12 @@ struct FSkeletalAnimMixerMutation : IMovieSceneEntityMutation
 		TComponentWriter<FMovieSceneRootMotionSettings> OutRootMotionSettings = Allocation->WriteComponents(AnimMixerComponents->RootMotionSettings, WriteContext);
 		TOptionalComponentReader<TObjectPtr<UMovieSceneAnimationMixerLayer>> OptMixerLayers = Allocation->TryReadComponents(AnimMixerComponents->MixerLayer);
 		
+		TComponentWriter<TObjectPtr<UMirrorDataTable>> OutMirrorTables = Allocation->WriteComponents(AnimMixerComponents->MirrorTable, WriteContext);
 		for (int32 Index = 0; Index < Num; ++Index)
 		{
 			// Initialize a null (empty) task, will be written during evaluation
 			OutTasks[Index] = nullptr;
+			OutMirrorTables[Index] = nullptr;
 
 			const FMovieSceneSkeletalAnimationComponentData& SkeletalAnim = SkeletalAnimationData[Index];
 			if (SkeletalAnim.Section)
@@ -146,6 +149,7 @@ struct FSkeletalAnimMixerMutation : IMovieSceneEntityMutation
 					LegacySwapRootBone = SkelAnimTrack->SwapRootBone;
 				}
 				OutRootMotionSettings[Index].LegacySwapRootBone = LegacySwapRootBone;
+				OutMirrorTables[Index] = SkeletalAnim.Section->Params.MirrorDataTable;
 
 				if (bCapturePreAnimatedState)
 				{
diff --git a/Engine/Plugins/Experimental/MovieSceneAnimMixer/Source/MovieSceneAnimMixer/Public/AnimMixerComponentTypes.h b/Engine/Plugins/Experimental/MovieSceneAnimMixer/Source/MovieSceneAnimMixer/Public/AnimMixerComponentTypes.h
--- a/Engine/Plugins/Experimental/MovieSceneAnimMixer/Source/MovieSceneAnimMixer/Public/AnimMixerComponentTypes.h
+++ b/Engine/Plugins/Experimental/MovieSceneAnimMixer/Source/MovieSceneAnimMixer/Public/AnimMixerComponentTypes.h
@@ -42,6 +42,7 @@ namespace UE::MovieScene
 		TComponentTypeID<FObjectComponent> MeshComponent;
 		TComponentTypeID<TSharedPtr<FMovieSceneMixerRootMotionComponentData>> MixerRootMotion;
 		TComponentTypeID<TObjectPtr<UMovieSceneAnimationMixerLayer>> MixerLayer;
+		TComponentTypeID<TObjectPtr<UMirrorDataTable>> MirrorTable;
 
 		struct
 		{
diff --git a/Engine/Plugins/Experimental/MovieSceneAnimMixer/Source/MovieSceneAnimMixer/Public/Systems/MovieSceneAnimMixerSystem.h b/Engine/Plugins/Experimental/MovieSceneAnimMixer/Source/MovieSceneAnimMixer/Public/Systems/MovieSceneAnimMixerSystem.h
--- a/Engine/Plugins/Experimental/MovieSceneAnimMixer/Source/MovieSceneAnimMixer/Public/Systems/MovieSceneAnimMixerSystem.h
+++ b/Engine/Plugins/Experimental/MovieSceneAnimMixer/Source/MovieSceneAnimMixer/Public/Systems/MovieSceneAnimMixerSystem.h
@@ -37,6 +37,7 @@ struct FMovieSceneAnimMixerEntry
 	int32 Priority = 0;
 	bool bAdditive = false;
 	bool bRequiresBlend = true;
+	TObjectPtr<UMirrorDataTable> MirrorTable;
 	UE::MovieScene::FMovieSceneEntityID EntityID;
 	UE::MovieScene::FInstanceHandle InstanceHandle;
 
diff --git a/Engine/Plugins/Experimental/UAF/UAFMirroring/Source/UAFMirroring/Internal/MirroringTraitData.h b/Engine/Plugins/Experimental/UAF/UAFMirroring/Source/UAFMirroring/Internal/MirroringTraitData.h
--- a/Engine/Plugins/Experimental/UAF/UAFMirroring/Source/UAFMirroring/Internal/MirroringTraitData.h
+++ b/Engine/Plugins/Experimental/UAF/UAFMirroring/Source/UAFMirroring/Internal/MirroringTraitData.h
@@ -4,57 +4,10 @@
 
 #include "TraitCore/TraitBinding.h"
 #include "TraitCore/TraitSharedData.h"
+#include "MirroringTask.h"
 
 #include "MirroringTraitData.generated.h"
 
-class UMirrorDataTable;
-
-/**
- * Mirroring Setup Parameters
- *
- * This struct holds the minimal configuration needed to enable a mirror pass.
- *
- * @see UMirrorDataTable, FMirroringTrait, FMirroringAdditiveTrait
- */
-USTRUCT(BlueprintType, meta = (DisplayName = "Mirroring Trait Setup Parameters"))
-struct FUAFMirroringTraitSetupParams
-{
-	GENERATED_BODY()
-
-	// Whether to perform mirror pass
-	UPROPERTY(EditAnywhere, Category = Settings, meta=(DisplayName = "Mirror"))
-	bool bShouldMirror = true;
-
-	// Data table to map bones to their mirrored counterpart
-	UPROPERTY(EditAnywhere, Category = Settings, meta = (ExportAsReference="true"))
-	TObjectPtr<const UMirrorDataTable> MirrorDataTable = nullptr;
-};
-
-/**
- * Mirroring Apply/Filter Parameters
- *
- * This struct holds flags for the channels that can be affected during the mirror pass (i.e. Bones, Curves, Attributes).
- *
- * @see FMirroringTrait, FMirroringAdditiveTrait, FAnimNextEvaluationMirroringTask
- */
-USTRUCT(BlueprintType, meta = (DisplayName = "Mirroring Trait Apply To Parameters"))
-struct FUAFMirroringTraitApplyToParams
-{
-	GENERATED_BODY()
-
-	// Whether to mirror bone transforms
-	UPROPERTY(EditAnywhere, Category = Channels, meta=(DisplayName = "Bones"))
-	bool bShouldMirrorBones = true;
-
-	// Whether to mirror animation curves
-	UPROPERTY(EditAnywhere, Category = Channels, meta=(DisplayName = "Curves"))
-	bool bShouldMirrorCurves = true;
-
-	// Whether to mirror attributes
-	UPROPERTY(EditAnywhere, Category = Channels, meta=(DisplayName = "Attributes"))
-	bool bShouldMirrorAttributes = true;
-};
-
 /** A trait that can mirror an input's keyframe data. */
 USTRUCT(meta = (DisplayName = "Mirroring", ShowTooltip = true))
 struct FMirroringTraitSharedData : public FAnimNextTraitSharedData
@@ -109,8 +62,6 @@ struct FMirroringAdditiveTraitSharedData : public FAnimNextTraitSharedData
 namespace UE::UAF
 {
 	// Namespaced alias
-	using FMirroringTraitSetupParams = FUAFMirroringTraitSetupParams;
-	using FMirroringTraitApplyToParams = FUAFMirroringTraitApplyToParams;
 	using FMirroringTraitData = FMirroringTraitSharedData;
 	using FMirroringAdditiveTraitData = FMirroringAdditiveTraitSharedData;
 	
diff --git a/Engine/Plugins/Experimental/UAF/UAFMirroring/Source/UAFMirroring/Private/MirroringTask.cpp b/Engine/Plugins/Experimental/UAF/UAFMirroring/Source/UAFMirroring/Private/MirroringTask.cpp
new file mode 100644
--- /dev/null
+++ b/Engine/Plugins/Experimental/UAF/UAFMirroring/Source/UAFMirroring/Private/MirroringTask.cpp
@@ -0,0 +1,161 @@
+// Copyright Epic Games, Inc. All Rights Reserved.
+
+#include "MirroringTask.h"
+
+#include "GenerationTools.h"
+#include "EvaluationVM/EvaluationVM.h"
+#include "TransformArrayOperations.h"
+#include "Animation/MirrorDataTable.h"
+#include "LODPose.h"
+#include "Animation/Skeleton.h"
+#include "Engine/SkeletalMesh.h"
+#include "AnimationRuntime.h"
+#include "Mirroring.h"
+#include "BoneContainer.h"
+
+////////////////////////////////////////////////////////////////////////////////////////////////////////////////////////
+// FAnimNextEvaluationMirroringTask
+
+void FAnimNextEvaluationMirroringTask::Execute(UE::UAF::FEvaluationVM& VM) const
+{
+	QUICK_SCOPE_CYCLE_COUNTER(FAnimNextEvaluationMirroringTask_Execute);
+
+	if (VM.GetActiveNamedSet())
+	{
+		// TODO: Implement with new attribute runtime
+		return;
+	}
+
+	if (!Setup.MirrorDataTable)
+	{
+		return;
+	}
+
+	// Pop our top pose and start processing it.
+	TUniquePtr<UE::UAF::FKeyframeState> Keyframe;
+	if (!VM.PopValue(UE::UAF::KEYFRAME_STACK_NAME, Keyframe))
+	{
+		// We have no inputs, nothing to do.
+		return;
+	}
+
+	// We don't support mirroring additive pose at the moment.
+	if (Keyframe->Pose.IsAdditive())
+	{
+		UE_LOG(LogAnimation, Warning, TEXT("FAnimNextEvaluationMirroringTask::Execute - Mirroring an additive pose is not supported."))
+	
+		// Force a bind pose to make it obvious.
+		Keyframe->Pose.SetIdentity(true);
+		VM.PushValue(UE::UAF::KEYFRAME_STACK_NAME, MoveTemp(Keyframe));
+		return;
+	}
+
+	// Rebuild mirror cache (if needed)
+	UE::UAF::FReferencePose ReferencePose = Keyframe->Pose.GetRefPose();
+	EnsureCache(VM, ReferencePose);
+
+	// Determine what channel to skip during this pass.
+	bool bShouldMirrorBones = ApplyTo.bShouldMirrorBones && EnumHasAnyFlags(VM.GetFlags(), UE::UAF::EEvaluationFlags::Bones);
+	bool bShouldMirrorCurves= ApplyTo.bShouldMirrorCurves && EnumHasAnyFlags(VM.GetFlags(), UE::UAF::EEvaluationFlags::Curves);
+	bool bShouldMirrorAttributes = ApplyTo.bShouldMirrorAttributes && EnumHasAnyFlags(VM.GetFlags(), UE::UAF::EEvaluationFlags::Attributes);
+
+	// After attempting to rebuild cache if its still invalid something went wrong.
+	if (!Cache.IsValid(ReferencePose, Setup.MirrorDataTable, bShouldMirrorBones, bShouldMirrorAttributes))
+	{
+		// Force a bind pose to make it obvious.
+		Keyframe->Pose.SetIdentity(Keyframe->Pose.IsAdditive());
+		VM.PushValue(UE::UAF::KEYFRAME_STACK_NAME, MoveTemp(Keyframe));
+		return;
+	}
+
+	// Mirror keyframe data.
+
+	if (bShouldMirrorBones)
+	{
+		UE::UAF::MirrorPose(
+			Keyframe->Pose,
+			Setup.MirrorDataTable->MirrorAxis,
+			Cache.MeshBoneIndexToMirroredMeshBoneIndexMap,
+			Cache.MeshSpaceReferencePoseRotations,
+			Cache.MeshSpaceReferenceRotationCorrections);
+	}
+
+	if (bShouldMirrorCurves)
+	{
+		FAnimationRuntime::MirrorCurves(Keyframe->Curves, *Setup.MirrorDataTable);
+	}
+
+	VM.PushValue(UE::UAF::KEYFRAME_STACK_NAME, MoveTemp(Keyframe));
+}
+
+void FAnimNextEvaluationMirroringTask::EnsureCache(const UE::UAF::FEvaluationVM& VM, const UE::UAF::FReferencePose& InReferencePose) const
+{
+	if (InReferencePose.SkeletalMesh.IsValid())
+	{
+		const int32 NumBonesForLOD0 = UE::UAF::GetNumOfBonesForMirrorData(InReferencePose);
+		const bool bShouldMirrorBones = ApplyTo.bShouldMirrorBones && EnumHasAnyFlags(VM.GetFlags(), UE::UAF::EEvaluationFlags::Bones);
+		const bool bShouldMirrorAttributes = ApplyTo.bShouldMirrorAttributes && EnumHasAnyFlags(VM.GetFlags(), UE::UAF::EEvaluationFlags::Attributes);
+		const bool bAreMirrorMapsValid = Cache.AreMirrorMapsValid(InReferencePose, Setup.MirrorDataTable, bShouldMirrorBones, bShouldMirrorAttributes);
+		const bool bIsReferencePoseDataValid = Cache.IsReferencePoseDataValid(InReferencePose, bShouldMirrorBones);
+	
+		if (bShouldMirrorBones)
+		{
+			if (!bAreMirrorMapsValid)
+			{
+				Cache.MeshBoneIndexToMirroredMeshBoneIndexMap.SetNumUninitialized(NumBonesForLOD0);
+			
+				UE::UAF::BuildMeshBoneIndexMirrorMap(InReferencePose, *Setup.MirrorDataTable, Cache.MeshBoneIndexToMirroredMeshBoneIndexMap);
+			}
+		
+			if (!bIsReferencePoseDataValid)
+			{
+				Cache.MeshSpaceReferencePoseRotations.SetNumUninitialized(NumBonesForLOD0);
+				Cache.MeshSpaceReferenceRotationCorrections.SetNumUninitialized(NumBonesForLOD0);
+			
+				UE::UAF::BuildReferencePoseMirrorData(
+					InReferencePose,
+					Setup.MirrorDataTable->MirrorAxis,
+					Cache.MeshBoneIndexToMirroredMeshBoneIndexMap,
+					Cache.MeshSpaceReferencePoseRotations,
+					Cache.MeshSpaceReferenceRotationCorrections
+				);
+			}
+		}
+	
+		if (bShouldMirrorAttributes)
+		{
+			if (!bAreMirrorMapsValid)
+			{
+				// We only care about the mesh bones (LOD0).
+				TArray<FBoneIndexType> RequiredBoneIndexArray;
+				RequiredBoneIndexArray.AddUninitialized(NumBonesForLOD0);
+			
+				for (int32 BoneIndex = 0; BoneIndex < RequiredBoneIndexArray.Num(); ++BoneIndex)
+				{
+					RequiredBoneIndexArray[BoneIndex] = InReferencePose.GetSkeletonBoneIndexFromLODBoneIndex(BoneIndex);
+				}
+
+				// Compute mirror map for compact pose.
+			
+				FBoneContainer BoneContainer;
+				const USkeleton* Skeleton = InReferencePose.SkeletalMesh->GetSkeleton();
+				BoneContainer.InitializeTo(RequiredBoneIndexArray, UE::Anim::FCurveFilterSettings(UE::Anim::ECurveFilterMode::DisallowAll), *Skeleton);
+			
+				TCustomBoneIndexArray<FSkeletonPoseBoneIndex, FSkeletonPoseBoneIndex> MirrorBoneIndexes;
+				Setup.MirrorDataTable->FillMirrorBoneIndexes(BoneContainer.GetSkeletonAsset(), MirrorBoneIndexes);
+				Setup.MirrorDataTable->FillCompactPoseMirrorBones(BoneContainer, MirrorBoneIndexes, Cache.CompactPoseBoneIndexToMirroredCompactPoseBoneIndexMap);
+			}
+		}
+	}
+	else
+	{
+		// Clear data to terminate task early.
+		Cache.Clear();
+	
+		UE_LOG(LogAnimation, Warning, TEXT("FAnimNextEvaluationMirroringTask::EnsureCache - Failed to get skeletal mesh asset from reference pose."))
+	}
+
+	// Keep track of latest assets used to build the cache
+	Cache.SkeletalMesh = InReferencePose.SkeletalMesh;
+	Cache.MirrorTable = Setup.MirrorDataTable;
+}
\ No newline at end of file
diff --git a/Engine/Plugins/Experimental/UAF/UAFMirroring/Source/UAFMirroring/Private/MirroringTrait.cpp b/Engine/Plugins/Experimental/UAF/UAFMirroring/Source/UAFMirroring/Private/MirroringTrait.cpp
--- a/Engine/Plugins/Experimental/UAF/UAFMirroring/Source/UAFMirroring/Private/MirroringTrait.cpp
+++ b/Engine/Plugins/Experimental/UAF/UAFMirroring/Source/UAFMirroring/Private/MirroringTrait.cpp
@@ -12,6 +12,7 @@
 #include "AnimationRuntime.h"
 #include "Mirroring.h"
 #include "BoneContainer.h"
+#include "MirroringTask.h"
 
 namespace UE::UAF
 {
@@ -119,6 +120,11 @@ namespace UE::UAF
 		{
 			InstanceData->Input = Context.AllocateNodeInstance(Binding, SharedData->Input);
 		}
+
+		if (!InstanceData->Task.IsValid())
+		{
+			InstanceData->Task = MakeShared<FAnimNextEvaluationMirroringTask>();
+		}
 	}
 
 	void FMirroringTrait::PostEvaluate(FEvaluateTraversalContext& Context, const TTraitBinding<IEvaluate>& Binding) const
@@ -131,16 +137,16 @@ namespace UE::UAF
 		// Copy since the settings properties are 'latent' and can change..
 		const FMirroringTraitSetupParams LatentSetupValue = SharedData->GetSetup(Binding);;
 
-		InstanceData->bHasMirrorStateChanged = LatentSetupValue.bShouldMirror != InstanceData->Setup.bShouldMirror;
-		InstanceData->Setup = LatentSetupValue;
-		InstanceData->ApplyTo = SharedData->GetApplyTo(Binding);
+		InstanceData->bHasMirrorStateChanged = LatentSetupValue.bShouldMirror != InstanceData->Task->Setup.bShouldMirror;
+		InstanceData->Task->Setup = LatentSetupValue;
+		InstanceData->Task->ApplyTo = SharedData->GetApplyTo(Binding);
 
-		if (!InstanceData->Setup.bShouldMirror || !InstanceData->Input.IsValid() || !InstanceData->Setup.MirrorDataTable)
+		if (!InstanceData->Task->Setup.bShouldMirror || !InstanceData->Input.IsValid() || !InstanceData->Task->Setup.MirrorDataTable)
 		{
 			return;
 		}
 		
-		Context.AppendTask(FAnimNextEvaluationMirroringTask::Make(InstanceData->Setup, InstanceData->ApplyTo, &InstanceData->Cache));
+		Context.AppendTaskPtr(InstanceData->Task);
 	}
 
 	uint32 FMirroringTrait::GetNumChildren(const FExecutionContext& Context, const TTraitBinding<IHierarchy>& Binding) const
@@ -162,7 +168,7 @@ namespace UE::UAF
 
 		FInstanceData* InstanceData = Binding.GetInstanceData<FInstanceData>();
 
-		Collector.AddReferencedObject(InstanceData->Setup.MirrorDataTable);
+		Collector.AddReferencedObject(InstanceData->Task->Setup.MirrorDataTable);
 	}
 
 	////////////////////////////////////////////////////////////////////////////////////////////////////////////////////////
@@ -178,6 +184,18 @@ namespace UE::UAF
 	GENERATE_ANIM_TRAIT_IMPLEMENTATION(FMirroringAdditiveTrait, TRAIT_INTERFACE_ENUMERATOR, NULL_ANIM_TRAIT_INTERFACE_ENUMERATOR, NULL_ANIM_TRAIT_EVENT_ENUMERATOR)
 	#undef TRAIT_INTERFACE_ENUMERATOR
 
+	void FMirroringAdditiveTrait::FInstanceData::Construct(const FExecutionContext& Context, const FTraitBinding& Binding)
+	{
+		FTrait::FInstanceData::Construct(Context, Binding);
+
+		FInstanceData* InstanceData = Binding.GetInstanceData<FInstanceData>();
+
+		if (!InstanceData->Task.IsValid())
+		{
+			InstanceData->Task = MakeShared<FAnimNextEvaluationMirroringTask>();
+		}
+	}
+
 	void FMirroringAdditiveTrait::PostEvaluate(FEvaluateTraversalContext& Context, const TTraitBinding<IEvaluate>& Binding) const
 	{
 		IEvaluate::PostEvaluate(Context, Binding);
@@ -188,16 +206,16 @@ namespace UE::UAF
 		// Copy since the settings properties are 'latent' and can change.
 		const FMirroringTraitSetupParams LatentSetupValue = SharedData->GetSetup(Binding);;
 
-		InstanceData->bHasMirrorStateChanged = LatentSetupValue.bShouldMirror != InstanceData->Setup.bShouldMirror;
-		InstanceData->Setup = LatentSetupValue;
-		InstanceData->ApplyTo = SharedData->GetApplyTo(Binding);
+		InstanceData->bHasMirrorStateChanged = LatentSetupValue.bShouldMirror != InstanceData->Task->Setup.bShouldMirror;
+		InstanceData->Task->Setup = LatentSetupValue;
+		InstanceData->Task->ApplyTo = SharedData->GetApplyTo(Binding);
 		
-		if (!InstanceData->Setup.bShouldMirror || !InstanceData->Setup.MirrorDataTable)
+		if (!InstanceData->Task->Setup.bShouldMirror || !InstanceData->Task->Setup.MirrorDataTable)
 		{
 			return;
 		}
 		
-		Context.AppendTask(FAnimNextEvaluationMirroringTask::Make(InstanceData->Setup, InstanceData->ApplyTo, &InstanceData->Cache));
+		Context.AppendTaskPtr(InstanceData->Task);
 	}
 
 	void FMirroringAdditiveTrait::AddReferencedObjects(const FExecutionContext& Context, const TTraitBinding<IGarbageCollection>& Binding, FReferenceCollector& Collector) const
@@ -206,162 +224,7 @@ namespace UE::UAF
 
 		FInstanceData* InstanceData = Binding.GetInstanceData<FInstanceData>();
 
-		Collector.AddReferencedObject(InstanceData->Setup.MirrorDataTable);
+		Collector.AddReferencedObject(InstanceData->Task->Setup.MirrorDataTable);
 	}
 	
 } // namespace UE::UAF
-
-////////////////////////////////////////////////////////////////////////////////////////////////////////////////////////
-// FAnimNextEvaluationMirroringTask
-
-FAnimNextEvaluationMirroringTask FAnimNextEvaluationMirroringTask::Make(const UE::UAF::FMirroringTraitSetupParams& Setup, const UE::UAF::FMirroringTraitApplyToParams& ApplyTo, UE::UAF::FMirroringTraitCache* Cache)
-{
-	FAnimNextEvaluationMirroringTask Task;
-	Task.Setup = Setup;
-	Task.ApplyTo = ApplyTo;
-	Task.Cache = Cache;
-	return Task;
-}
-
-void FAnimNextEvaluationMirroringTask::Execute(UE::UAF::FEvaluationVM& VM) const
-{
-	QUICK_SCOPE_CYCLE_COUNTER(FAnimNextEvaluationMirroringTask_Execute);
-
-	if (VM.GetActiveNamedSet())
-	{
-		// TODO: Implement with new attribute runtime
-		return;
-	}
-
-	// Pop our top pose and start processing it.
-	TUniquePtr<UE::UAF::FKeyframeState> Keyframe;
-	if (!VM.PopValue(UE::UAF::KEYFRAME_STACK_NAME, Keyframe))
-	{
-		// We have no inputs, nothing to do.
-		return;
-	}
-
-	// We don't support mirroring additive pose at the moment.
-	if (Keyframe->Pose.IsAdditive())
-	{
-		UE_LOG(LogAnimation, Warning, TEXT("FAnimNextEvaluationMirroringTask::Execute - Mirroring an additive pose is not supported."))
-		
-		// Force a bind pose to make it obvious.
-		Keyframe->Pose.SetIdentity(true);
-		VM.PushValue(UE::UAF::KEYFRAME_STACK_NAME, MoveTemp(Keyframe));
-		return;
-	}
-
-	// Rebuild mirror cache (if needed)
-	UE::UAF::FReferencePose ReferencePose = Keyframe->Pose.GetRefPose();
-	EnsureCache(VM, ReferencePose);
-
-	// Determine what channel to skip during this pass.
-	bool bShouldMirrorBones = ApplyTo.bShouldMirrorBones && EnumHasAnyFlags(VM.GetFlags(), UE::UAF::EEvaluationFlags::Bones);
-	bool bShouldMirrorCurves= ApplyTo.bShouldMirrorCurves && EnumHasAnyFlags(VM.GetFlags(), UE::UAF::EEvaluationFlags::Curves);
-	bool bShouldMirrorAttributes = ApplyTo.bShouldMirrorAttributes && EnumHasAnyFlags(VM.GetFlags(), UE::UAF::EEvaluationFlags::Attributes);
-	
-	// After attempting to rebuild cache if its still invalid something went wrong.
-	if (!Cache->IsValid(ReferencePose, Setup.MirrorDataTable, bShouldMirrorBones, bShouldMirrorAttributes))
-	{
-		// Force a bind pose to make it obvious.
-		Keyframe->Pose.SetIdentity(Keyframe->Pose.IsAdditive());
-		VM.PushValue(UE::UAF::KEYFRAME_STACK_NAME, MoveTemp(Keyframe));
-		return;
-	}
-
-	// Mirror keyframe data.
-	
-	if (bShouldMirrorBones)
-	{
-		UE::UAF::MirrorPose(
-			Keyframe->Pose,
-			Setup.MirrorDataTable->MirrorAxis,
-			Cache->MeshBoneIndexToMirroredMeshBoneIndexMap,
-			Cache->MeshSpaceReferencePoseRotations,
-			Cache->MeshSpaceReferenceRotationCorrections);
-	}
-
-	if (bShouldMirrorCurves)
-	{
-		FAnimationRuntime::MirrorCurves(Keyframe->Curves, *Setup.MirrorDataTable);
-	}
-
-
-	VM.PushValue(UE::UAF::KEYFRAME_STACK_NAME, MoveTemp(Keyframe));
-}
-
-void FAnimNextEvaluationMirroringTask::EnsureCache(const UE::UAF::FEvaluationVM& VM, const UE::UAF::FReferencePose& InReferencePose) const
-{
-	checkf(Cache != nullptr, TEXT("FAnimNextEvaluationMirroringTask::EnsureCache - Invalid cache ptr."))
-	
-	if (InReferencePose.SkeletalMesh.IsValid())
-	{
-		const int32 NumBonesForLOD0 = UE::UAF::GetNumOfBonesForMirrorData(InReferencePose);
-		const bool bShouldMirrorBones = ApplyTo.bShouldMirrorBones && EnumHasAnyFlags(VM.GetFlags(), UE::UAF::EEvaluationFlags::Bones);
-		const bool bShouldMirrorAttributes = ApplyTo.bShouldMirrorAttributes && EnumHasAnyFlags(VM.GetFlags(), UE::UAF::EEvaluationFlags::Attributes);
-		const bool bAreMirrorMapsValid = Cache->AreMirrorMapsValid(InReferencePose, Setup.MirrorDataTable, bShouldMirrorBones, bShouldMirrorAttributes);
-		const bool bIsReferencePoseDataValid = Cache->IsReferencePoseDataValid(InReferencePose, bShouldMirrorBones);
-		
-		if (bShouldMirrorBones)
-		{
-			if (!bAreMirrorMapsValid)
-			{
-				Cache->MeshBoneIndexToMirroredMeshBoneIndexMap.SetNumUninitialized(NumBonesForLOD0);
-				UE::UAF::BuildMeshBoneIndexMirrorMap(InReferencePose, *Setup.MirrorDataTable, Cache->MeshBoneIndexToMirroredMeshBoneIndexMap);
-			}
-			
-			if (!bIsReferencePoseDataValid)
-			{
-				Cache->MeshSpaceReferencePoseRotations.SetNumUninitialized(NumBonesForLOD0);
-				Cache->MeshSpaceReferenceRotationCorrections.SetNumUninitialized(NumBonesForLOD0);
-				
-				UE::UAF::BuildReferencePoseMirrorData(
-					InReferencePose,
-					Setup.MirrorDataTable->MirrorAxis,
-					Cache->MeshBoneIndexToMirroredMeshBoneIndexMap,
-					Cache->MeshSpaceReferencePoseRotations,
-					Cache->MeshSpaceReferenceRotationCorrections
-				);
-			}
-
-			if (bShouldMirrorAttributes)
-			{
-				if (!bAreMirrorMapsValid)
-				{
-					// We only care about the mesh bones (LOD0).
-					TArray<FBoneIndexType> RequiredBoneIndexArray;
-					RequiredBoneIndexArray.AddUninitialized(NumBonesForLOD0);
-
-					for (int32 BoneIndex = 0; BoneIndex < RequiredBoneIndexArray.Num(); ++BoneIndex)
-					{
-						RequiredBoneIndexArray[BoneIndex] = InReferencePose.GetSkeletonBoneIndexFromLODBoneIndex(BoneIndex);
-					}
-
-					// Compute mirror map for compact pose.
-
-					FBoneContainer BoneContainer;
-					const USkeleton* Skeleton = InReferencePose.SkeletalMesh->GetSkeleton();
-					BoneContainer.InitializeTo(RequiredBoneIndexArray, UE::Anim::FCurveFilterSettings(UE::Anim::ECurveFilterMode::DisallowAll), *Skeleton);
-
-					TCustomBoneIndexArray<FSkeletonPoseBoneIndex, FSkeletonPoseBoneIndex> MirrorBoneIndexes;
-					Setup.MirrorDataTable->FillMirrorBoneIndexes(BoneContainer.GetSkeletonAsset(), MirrorBoneIndexes);
-					Setup.MirrorDataTable->FillCompactPoseMirrorBones(BoneContainer, MirrorBoneIndexes, Cache->CompactPoseBoneIndexToMirroredCompactPoseBoneIndexMap);
-				}
-			}
-		}
-		
-	
-	}
-	else
-	{
-		// Clear data to terminate task early.
-		Cache->Clear();
-		
-		UE_LOG(LogAnimation, Warning, TEXT("FAnimNextEvaluationMirroringTask::EnsureCache - Failed to get skeletal mesh asset from reference pose."))
-	}
-
-	// Keep track of latest assets used to build the cache
-	Cache->SkeletalMesh = InReferencePose.SkeletalMesh;
-	Cache->MirrorTable = Setup.MirrorDataTable;
-}
\ No newline at end of file
diff --git a/Engine/Plugins/Experimental/UAF/UAFMirroring/Source/UAFMirroring/Private/MirroringTrait.h b/Engine/Plugins/Experimental/UAF/UAFMirroring/Source/UAFMirroring/Private/MirroringTrait.h
--- a/Engine/Plugins/Experimental/UAF/UAFMirroring/Source/UAFMirroring/Private/MirroringTrait.h
+++ b/Engine/Plugins/Experimental/UAF/UAFMirroring/Source/UAFMirroring/Private/MirroringTrait.h
@@ -14,62 +14,11 @@
 #include "MirroringTraitData.h"
 #include "CustomBoneIndexArray.h"
 
-#include "MirroringTrait.generated.h"
-
 class UMirrorDataTable;
 class USkeletalMesh;
 
 namespace UE::UAF
 {
-	/**
-	 * Mirroring Cache
-	 *
-	 * Holds all precomputed data needed to mirror keyframe output efficiently, so it doesn’t have to be rebuilt every evaluation.
-	 * 
-	 * @see FAnimNextEvaluationMirroringTask
-	 */
-	struct FMirroringTraitCache
-	{
-		// Cached mirror indices map (invalidated on skeletal mesh or mirror data table change)
-		TArray<FBoneIndexType> MeshBoneIndexToMirroredMeshBoneIndexMap;
-		
-		// Cached bind pose rotations (invalidated on skeletal mesh change)
-		TArray<FQuat> MeshSpaceReferencePoseRotations;
-
-		// Cached bind pose rotation corrections (invalidated on skeletal mesh change)
-		TArray<FQuat> MeshSpaceReferenceRotationCorrections;
-		
-		// Cached mirror indices map but using FCompactPoseIndex (invalidated on skeletal mesh or mirror data table change)
-		TCustomBoneIndexArray<FCompactPoseBoneIndex, FCompactPoseBoneIndex> CompactPoseBoneIndexToMirroredCompactPoseBoneIndexMap;
-
-		// Skeletal mesh used to build this cache
-		TWeakObjectPtr<const USkeletalMesh> SkeletalMesh = nullptr;
-
-		// Mirror data table used to build this cache
-		TWeakObjectPtr<const UMirrorDataTable> MirrorTable = nullptr;
-
-		/**
-		 * True if the cached mirror maps were generated using exactly these assets.
-		 * Any mismatch (different mesh/mirror table, size mismatch) means you must rebuild the mirror maps.
-		 */
-		bool AreMirrorMapsValid(const UE::UAF::FReferencePose& InReferencePose, const TWeakObjectPtr<const UMirrorDataTable>& InMirrorTable, bool bInShouldMirrorBones, bool bInShouldMirrorAttributes) const;
-
-		/**
-		 * True if the cached bind/reference-pose data matches this mesh.
-		 * Any mismatch (different mesh table, size mismatch)  means you must rebuild the reference-pose arrays.
-		 */
-		bool IsReferencePoseDataValid(const UE::UAF::FReferencePose& InReferencePose, bool bInShouldMirrorBones) const;
-
-		/**
-		 * True if the cache was generated using exactly these assets.
-		 * Any mismatch (different mesh/mirror table, size mismatch) means you must rebuild the cache.
-		 */
-		bool IsValid(const UE::UAF::FReferencePose& InReferencePose, const TWeakObjectPtr<const UMirrorDataTable>& InMirrorTable, bool bInShouldMirrorBones, bool bInShouldMirrorAttributes) const;
-
-		/** Empty arrays and reset ptrs to assets */
-		void Clear();
-	};
-	
 	/**
 	 * Mirroring Base Trait
 	 * 
@@ -85,19 +34,13 @@ namespace UE::UAF
 		{
 			// Input node to query the keyframe to be mirrored.
 			FTraitPtr Input;
-			
-			// Defines whether to perform mirror pass (and what data table to use).
-			FMirroringTraitSetupParams Setup;
-
-			// Defines what channels will be affected during the mirror pass.
-			FMirroringTraitApplyToParams ApplyTo;
 
+			// Task containing the setup data on how to do the mirroring
+			TSharedPtr<FAnimNextEvaluationMirroringTask> Task;
+			
 			// Used to trigger inertial blends.
 			bool bHasMirrorStateChanged = false;
 			
-			// Data that stays the same between evaluation passes.
-			FMirroringTraitCache Cache;
-			
 			void Construct(const FExecutionContext& Context, const FTraitBinding& Binding);
 		};
 		
@@ -125,17 +68,13 @@ namespace UE::UAF
 
 		struct FInstanceData : FTrait::FInstanceData
 		{
-			// Defines whether to perform mirror pass (and what data table to use).
-			FMirroringTraitSetupParams Setup;
-
-			// Defines what channels will be affected during the mirror pass.
-			FMirroringTraitApplyToParams ApplyTo;
-			
+			TSharedPtr<FAnimNextEvaluationMirroringTask> Task;
+				
 			// Used to trigger inertial blends.
 			bool bHasMirrorStateChanged = false;
+
+			void Construct(const FExecutionContext& Context, const FTraitBinding& Binding);
 			
-			// Data that stays the same between evaluation passes.
-			FMirroringTraitCache Cache;
 		};
 
 		// IEvaluate impl
@@ -145,39 +84,4 @@ namespace UE::UAF
 		virtual void AddReferencedObjects(const FExecutionContext& Context, const TTraitBinding<IGarbageCollection>& Binding, FReferenceCollector& Collector) const override;
 	};
 	
-} // namespace UE::UAF
-
-/**
- * Mirroring Task
- * 
- * This pop the top keyframe from the VM keyframe stack, mirrors it, and pushes
- * back the result onto the stack.
- */
-USTRUCT()
-struct FAnimNextEvaluationMirroringTask : public FAnimNextEvaluationTask
-{
-	GENERATED_BODY()
-
-	DECLARE_ANIM_EVALUATION_TASK(FAnimNextEvaluationMirroringTask)
-
-	// Make a FAnimNextEvaluationMirroringTask task
-	static FAnimNextEvaluationMirroringTask Make(const UE::UAF::FMirroringTraitSetupParams & Setup, const UE::UAF::FMirroringTraitApplyToParams & ApplyTo, UE::UAF::FMirroringTraitCache* Cache);
-	
-	// Settings to use when doing the mirror pass
-	UPROPERTY(VisibleAnywhere, Category = "Setup")
-	FUAFMirroringTraitSetupParams Setup;
-
-	UPROPERTY(VisibleAnywhere, Category = "Apply To")
-	FUAFMirroringTraitApplyToParams ApplyTo;
-
-	// Cache for bind pose rotation, corrections, and mirror map
-	UE::UAF::FMirroringTraitCache* Cache = nullptr;
-	
-	// Task entry point
-	virtual void Execute(UE::UAF::FEvaluationVM& VM) const override;
-	
-private:
-
-	// Ensures the mirror cache matches the given inputs, rebuilds if stale.
-	void EnsureCache(const UE::UAF::FEvaluationVM& VM, const UE::UAF::FReferencePose& InReferencePose) const;
-};
\ No newline at end of file
+} // namespace UE::UAF
\ No newline at end of file
diff --git a/Engine/Plugins/Experimental/UAF/UAFMirroring/Source/UAFMirroring/Public/MirroringTask.h b/Engine/Plugins/Experimental/UAF/UAFMirroring/Source/UAFMirroring/Public/MirroringTask.h
new file mode 100644
--- /dev/null
+++ b/Engine/Plugins/Experimental/UAF/UAFMirroring/Source/UAFMirroring/Public/MirroringTask.h
@@ -0,0 +1,145 @@
+// Copyright Epic Games, Inc. All Rights Reserved.
+
+#pragma once
+
+#include "CoreMinimal.h"
+#include "EvaluationVM/EvaluationTask.h"
+#include "ReferencePose.h"
+#include "CustomBoneIndexArray.h"
+#include "MirroringTask.generated.h"
+
+class UMirrorDataTable;
+
+/**
+	 * Mirroring Setup Parameters
+	 *
+	 * This struct holds the minimal configuration needed to enable a mirror pass.
+	 *
+	 * @see UMirrorDataTable, FMirroringTrait, FMirroringAdditiveTrait
+	 */
+USTRUCT(BlueprintType, meta = (DisplayName = "Mirroring Trait Setup Parameters"))
+struct FUAFMirroringTraitSetupParams
+{
+	GENERATED_BODY()
+
+	// Whether to perform mirror pass
+	UPROPERTY(EditAnywhere, Category = Settings, meta=(DisplayName = "Mirror"))
+	bool bShouldMirror = true;
+
+	// Data table to map bones to their mirrored counterpart
+	UPROPERTY(EditAnywhere, Category = Settings, meta = (ExportAsReference="true"))
+	TObjectPtr<const UMirrorDataTable> MirrorDataTable = nullptr;
+};
+
+/**
+ * Mirroring Apply/Filter Parameters
+ *
+ * This struct holds flags for the channels that can be affected during the mirror pass (i.e. Bones, Curves, Attributes).
+ *
+ * @see FMirroringTrait, FMirroringAdditiveTrait, FAnimNextEvaluationMirroringTask
+ */
+USTRUCT(BlueprintType, meta = (DisplayName = "Mirroring Trait Apply To Parameters"))
+struct FUAFMirroringTraitApplyToParams
+{
+	GENERATED_BODY()
+
+	// Whether to mirror bone transforms
+	UPROPERTY(EditAnywhere, Category = Channels, meta=(DisplayName = "Bones"))
+	bool bShouldMirrorBones = true;
+
+	// Whether to mirror animation curves
+	UPROPERTY(EditAnywhere, Category = Channels, meta=(DisplayName = "Curves"))
+	bool bShouldMirrorCurves = true;
+
+	// Whether to mirror attributes
+	UPROPERTY(EditAnywhere, Category = Channels, meta=(DisplayName = "Attributes"))
+	bool bShouldMirrorAttributes = true;
+};
+
+namespace UE::UAF
+{
+	// Namespaced alias
+	using FMirroringTraitSetupParams = FUAFMirroringTraitSetupParams;
+	using FMirroringTraitApplyToParams = FUAFMirroringTraitApplyToParams;
+	
+	/**
+	 * Mirroring Cache
+	 *
+	 * Holds all precomputed data needed to mirror keyframe output efficiently, so it doesn’t have to be rebuilt every evaluation.
+	 * 
+	 * @see FAnimNextEvaluationMirroringTask
+	 */
+	struct FMirroringTraitCache
+	{
+		// Cached mirror indices map (invalidated on skeletal mesh or mirror data table change)
+		TArray<FBoneIndexType> MeshBoneIndexToMirroredMeshBoneIndexMap;
+		
+		// Cached bind pose rotations (invalidated on skeletal mesh change)
+		TArray<FQuat> MeshSpaceReferencePoseRotations;
+
+		// Cached bind pose rotation corrections (invalidated on skeletal mesh change)
+		TArray<FQuat> MeshSpaceReferenceRotationCorrections;
+		
+		// Cached mirror indices map but using FCompactPoseIndex (invalidated on skeletal mesh or mirror data table change)
+		TCustomBoneIndexArray<FCompactPoseBoneIndex, FCompactPoseBoneIndex> CompactPoseBoneIndexToMirroredCompactPoseBoneIndexMap;
+
+		// Skeletal mesh used to build this cache
+		TWeakObjectPtr<const USkeletalMesh> SkeletalMesh = nullptr;
+
+		// Mirror data table used to build this cache
+		TWeakObjectPtr<const UMirrorDataTable> MirrorTable = nullptr;
+
+		/**
+		 * True if the cached mirror maps were generated using exactly these assets.
+		 * Any mismatch (different mesh/mirror table, size mismatch) means you must rebuild the mirror maps.
+		 */
+		bool AreMirrorMapsValid(const UE::UAF::FReferencePose& InReferencePose, const TWeakObjectPtr<const UMirrorDataTable>& InMirrorTable, bool bInShouldMirrorBones, bool bInShouldMirrorAttributes) const;
+
+		/**
+		 * True if the cached bind/reference-pose data matches this mesh.
+		 * Any mismatch (different mesh table, size mismatch)  means you must rebuild the reference-pose arrays.
+		 */
+		bool IsReferencePoseDataValid(const UE::UAF::FReferencePose& InReferencePose, bool bInShouldMirrorBones) const;
+
+		/**
+		 * True if the cache was generated using exactly these assets.
+		 * Any mismatch (different mesh/mirror table, size mismatch) means you must rebuild the cache.
+		 */
+		bool IsValid(const UE::UAF::FReferencePose& InReferencePose, const TWeakObjectPtr<const UMirrorDataTable>& InMirrorTable, bool bInShouldMirrorBones, bool bInShouldMirrorAttributes) const;
+
+		/** Empty arrays and reset ptrs to assets */
+		void Clear();
+	};
+}
+	
+/**
+* Mirroring Task
+ * 
+ * This pop the top keyframe from the VM keyframe stack, mirrors it, and pushes
+ * back the result onto the stack.
+ */
+USTRUCT()
+struct UAFMIRRORING_API FAnimNextEvaluationMirroringTask : public FAnimNextEvaluationTask
+{
+	GENERATED_BODY()
+
+	DECLARE_ANIM_EVALUATION_TASK(FAnimNextEvaluationMirroringTask)
+
+	// Settings to use when doing the mirror pass
+	UPROPERTY(VisibleAnywhere, Category = "Setup")
+	FUAFMirroringTraitSetupParams Setup;
+
+	UPROPERTY(VisibleAnywhere, Category = "Apply To")
+	FUAFMirroringTraitApplyToParams ApplyTo;
+
+	// Cache for bind pose rotation, corrections, and mirror map
+	mutable UE::UAF::FMirroringTraitCache Cache;
+
+	// Task entry point
+	virtual void Execute(UE::UAF::FEvaluationVM& VM) const override;
+
+private:
+
+	// Ensures the mirror cache matches the given inputs, rebuilds if stale.
+	void EnsureCache(const UE::UAF::FEvaluationVM& VM, const UE::UAF::FReferencePose& InReferencePose) const;
+};
\ No newline at end of file
