UAF fix some incorrect comparison of invalid bone indicies, where 16bit was upcast to 32bit and compared to -1, which will fail
Add a define for INVALID_BONE_INDEX that is 16bit and convert UAF code to use it, to help prevent this mistake in the future.

#Synced-CL 50067831

#Synced-CL 50067831

diff --git a/Engine/Plugins/Experimental/UAF/UAF/Source/UAF/Public/LODPose.h b/Engine/Plugins/Experimental/UAF/UAF/Source/UAF/Public/LODPose.h
--- a/Engine/Plugins/Experimental/UAF/UAF/Source/UAF/Public/LODPose.h
+++ b/Engine/Plugins/Experimental/UAF/UAF/Source/UAF/Public/LODPose.h
@@ -224,14 +224,14 @@ struct FLODPose
 	}
 
 	// Query to find a LODBoneIndex for an associated BoneName
-	// Returns INDEX_NONE if missing ReferencePose or if no bone found for a given name
+	// Returns INVALID_BONE_INDEX if missing ReferencePose or if no bone found for a given name
 	FBoneIndexType FindLODBoneIndexFromBoneName(FName BoneName) const
 	{
 		if (RefPose != nullptr)
 		{
 			return RefPose->FindLODBoneIndexFromBoneName(BoneName);
 		}
-		return INDEX_NONE;
+		return INVALID_BONE_INDEX;
 	}
 
 	// Query whether bone with LODIndex ChildLODBoneIndex is a child of bone with LODIndex ParentLODBoneIndex
@@ -240,8 +240,8 @@ struct FLODPose
 	//  False - ChildLODBoneIndex is not a child of ParentLODBoneIndex
 	bool IsBoneChildOf(FBoneIndexType ChildLODBoneIndex, FBoneIndexType ParentLODBoneIndex) const
 	{
-		check(ChildLODBoneIndex != INDEX_NONE);
-		check(ParentLODBoneIndex != INDEX_NONE);
+		check(ChildLODBoneIndex != INVALID_BONE_INDEX);
+		check(ParentLODBoneIndex != INVALID_BONE_INDEX);
 		check(ChildLODBoneIndex != ParentLODBoneIndex);
 		
 		if (RefPose == nullptr)
@@ -252,7 +252,7 @@ struct FLODPose
 		const TArrayView<const FBoneIndexType> LODBoneIndexToParentMap = RefPose->GetLODBoneIndexToParentLODBoneIndexMap(LODLevel); 
 		
 		FBoneIndexType ParentOfChild = LODBoneIndexToParentMap[ChildLODBoneIndex];
-		while (ParentOfChild != (FBoneIndexType)INDEX_NONE)	// Note: Explicit cast to make static analysis happy
+		while (ParentOfChild != INVALID_BONE_INDEX)
 		{
 			if (ParentOfChild == ParentLODBoneIndex)
 			{
@@ -268,12 +268,12 @@ struct FLODPose
 	}
 
 	// Get the LODBoneIndex of the parent of the bone at ChildLODBoneIndex
-	// Returns INDEX_NONE if no reference pose or if ChilddLODBoneIndex is the root
+	// Returns INVALID_BONE_INDEX if no reference pose or if ChilddLODBoneIndex is the root
 	FBoneIndexType GetLODBoneParentIndex(FBoneIndexType ChildLODBoneIndex) const
 	{
 		if (RefPose == nullptr)
 		{
-			return INDEX_NONE;
+			return INVALID_BONE_INDEX;
 		}
 		
 		return RefPose->GetLODParentBoneIndex(LODLevel, ChildLODBoneIndex);
diff --git a/Engine/Plugins/Experimental/UAF/UAF/Source/UAF/Public/ReferencePose.h b/Engine/Plugins/Experimental/UAF/UAF/Source/UAF/Public/ReferencePose.h
--- a/Engine/Plugins/Experimental/UAF/UAF/Source/UAF/Public/ReferencePose.h
+++ b/Engine/Plugins/Experimental/UAF/UAF/Source/UAF/Public/ReferencePose.h
@@ -214,17 +214,17 @@ struct TReferencePose
 	}
 
 	// Query to find a LODBoneIndex for an associated BoneName
-	// Returns INDEX_NONE if no bone found for a given name
+	// Returns INVALID_BONE_INDEX if no bone found for a given name
 	FBoneIndexType FindLODBoneIndexFromBoneName(FName BoneName) const
 	{
 		if (const FBoneIndexType* BoneIndex = BoneNameToLODBoneIndexMap.Find(BoneName))
 		{
 			return *BoneIndex;
 		}
-		return INDEX_NONE;
+		return INVALID_BONE_INDEX;
 	}
 
-	int32 GetMeshBoneIndexFromLODBoneIndex(int32 LODBoneIndex) const
+	int32 GetMeshBoneIndexFromLODBoneIndex(FBoneIndexType LODBoneIndex) const
 	{
 		const int32 NumBonesLOD0 = LODBoneIndexToMeshBoneIndexMapPerLOD[0].Num();
 		check(LODBoneIndex < NumBonesLOD0);
@@ -237,8 +237,8 @@ struct TReferencePose
 		return MeshBoneIndexToParentMeshBoneIndexMap;
 	}
 
-	// Returns the corresponding Skeleron Bone Index to a LOD Bone Index
-	int32 GetSkeletonBoneIndexFromLODBoneIndex(int32 LODBoneIndex) const
+	// Returns the corresponding skeleton Bone Index to a LOD Bone Index
+	FBoneIndexType GetSkeletonBoneIndexFromLODBoneIndex(FBoneIndexType LODBoneIndex) const
 	{
 		const int32 NumBonesLOD0 = LODBoneIndexToSkeletonBoneIndexMapPerLOD[0].Num();
 		check(LODBoneIndex < NumBonesLOD0);
@@ -246,19 +246,19 @@ struct TReferencePose
 	}
 
 	// Translate a skeletal mesh bone index to a pose bone index.
-	int32 GetLODBoneIndexFromMeshBoneIndex(int32 MeshBoneIndex) const
+	FBoneIndexType GetLODBoneIndexFromMeshBoneIndex(FBoneIndexType MeshBoneIndex) const
 	{
 		check(MeshBoneIndex < MeshBoneIndexToLODBoneIndexMap.Num());
 		return MeshBoneIndexToLODBoneIndexMap[MeshBoneIndex];
 	}
 
-	int32 GetLODBoneIndexFromSkeletonBoneIndex(int32 SkeletionBoneIndex) const
+	FBoneIndexType GetLODBoneIndexFromSkeletonBoneIndex(FBoneIndexType SkeletionBoneIndex) const
 	{
 		check(SkeletionBoneIndex < SkeletonBoneIndexToLODBoneIndexMap.Num());
 		return SkeletonBoneIndexToLODBoneIndexMap[SkeletionBoneIndex];
 	}
 
-	int32 GetLODParentBoneIndex(int32 LODLevel, int32 LODBoneIndex) const
+	FBoneIndexType GetLODParentBoneIndex(int32 LODLevel, FBoneIndexType LODBoneIndex) const
 	{
 		const TArrayView<const FBoneIndexType> ParentLODMap = GetLODBoneIndexToParentLODBoneIndexMap(LODLevel);
 		if (ParentLODMap.IsValidIndex(LODBoneIndex))
@@ -269,7 +269,7 @@ struct TReferencePose
 				return LODParentBoneIndex;
 			}
 		}
-		return INDEX_NONE;
+		return INVALID_BONE_INDEX;
 	}
 
 	const USkeletalMesh* GetSkeletalMeshAsset() const
@@ -281,7 +281,7 @@ struct TReferencePose
 		return Skeleton.Get();
 	}
 
-	FTransform GetRefPoseTransform(int32 LODBoneIndex) const
+	FTransform GetRefPoseTransform(FBoneIndexType LODBoneIndex) const
 	{
 		const int32 NumBonesLOD0 = LODBoneIndexToMeshBoneIndexMapPerLOD[0].Num();
 		check(LODBoneIndex < NumBonesLOD0);
@@ -289,7 +289,7 @@ struct TReferencePose
 		return ReferenceLocalTransforms[LODBoneIndex];
 	}
 
-	const FQuat GetRefPoseRotation(int32 LODBoneIndex) const
+	const FQuat GetRefPoseRotation(FBoneIndexType LODBoneIndex) const
 	{
 		const int32 NumBonesLOD0 = LODBoneIndexToMeshBoneIndexMapPerLOD[0].Num();
 		check(LODBoneIndex < NumBonesLOD0);
@@ -297,7 +297,7 @@ struct TReferencePose
 		return ReferenceLocalTransforms[LODBoneIndex].GetRotation();
 	}
 
-	const FVector GetRefPoseTranslation(int32 LODBoneIndex) const
+	const FVector GetRefPoseTranslation(FBoneIndexType LODBoneIndex) const
 	{
 		const int32 NumBonesLOD0 = LODBoneIndexToMeshBoneIndexMapPerLOD[0].Num();
 		check(LODBoneIndex < NumBonesLOD0);
@@ -305,7 +305,7 @@ struct TReferencePose
 		return ReferenceLocalTransforms[LODBoneIndex].GetTranslation();
 	}
 
-	const FVector GetRefPoseScale3D(int32 LODBoneIndex) const
+	const FVector GetRefPoseScale3D(FBoneIndexType LODBoneIndex) const
 	{
 		const int32 NumBonesLOD0 = LODBoneIndexToMeshBoneIndexMapPerLOD[0].Num();
 		check(LODBoneIndex < NumBonesLOD0);
diff --git a/Engine/Plugins/Experimental/UAF/UAF/Source/UAF/Public/RemapPoseData.h b/Engine/Plugins/Experimental/UAF/UAF/Source/UAF/Public/RemapPoseData.h
--- a/Engine/Plugins/Experimental/UAF/UAF/Source/UAF/Public/RemapPoseData.h
+++ b/Engine/Plugins/Experimental/UAF/UAF/Source/UAF/Public/RemapPoseData.h
@@ -14,8 +14,8 @@ struct FRemappedBone
 {
 	GENERATED_BODY()
 
-	FBoneIndexType SourceBoneIndex = INDEX_NONE;
-	FBoneIndexType TargetBoneIndex = INDEX_NONE;
+	FBoneIndexType SourceBoneIndex = INVALID_BONE_INDEX;
+	FBoneIndexType TargetBoneIndex = INVALID_BONE_INDEX;
 };
 
 USTRUCT()
diff --git a/Engine/Plugins/Experimental/UAF/UAFAnimGraph/Source/UAFAnimGraph/Private/EvaluationVM/Tasks/CopyBones.cpp b/Engine/Plugins/Experimental/UAF/UAFAnimGraph/Source/UAFAnimGraph/Private/EvaluationVM/Tasks/CopyBones.cpp
--- a/Engine/Plugins/Experimental/UAF/UAFAnimGraph/Source/UAFAnimGraph/Private/EvaluationVM/Tasks/CopyBones.cpp
+++ b/Engine/Plugins/Experimental/UAF/UAFAnimGraph/Source/UAFAnimGraph/Private/EvaluationVM/Tasks/CopyBones.cpp
@@ -43,8 +43,8 @@ void FAnimNextCopyBonesComponentSpaceTask::Execute(UE::UAF::FEvaluationVM& VM) c
 
 				// If we can't find them or they are not in the current LOD then skip
 
-				if (SourceBoneIdx == FBoneIndexType(INDEX_NONE) ||
-					DestinationBoneIdx == FBoneIndexType(INDEX_NONE) ||
+				if (SourceBoneIdx == INVALID_BONE_INDEX ||
+					DestinationBoneIdx == INVALID_BONE_INDEX ||
 					SourceBoneIdx >= Keyframe->Get()->Pose.GetNumBones() ||
 					DestinationBoneIdx >= Keyframe->Get()->Pose.GetNumBones())
 				{
@@ -54,8 +54,8 @@ void FAnimNextCopyBonesComponentSpaceTask::Execute(UE::UAF::FEvaluationVM& VM) c
 				// Compute the Source Component Space Transform
 
 				FTransform SourceComponentTransform = Keyframe->Get()->Pose.LocalTransforms[SourceBoneIdx];
-				int32 SourceParentIdx = Keyframe->Get()->Pose.GetRefPose().GetLODParentBoneIndex(Keyframe->Get()->Pose.LODLevel, SourceBoneIdx);
-				while (SourceParentIdx != INDEX_NONE)
+				FBoneIndexType SourceParentIdx = Keyframe->Get()->Pose.GetRefPose().GetLODParentBoneIndex(Keyframe->Get()->Pose.LODLevel, SourceBoneIdx);
+				while (SourceParentIdx != INVALID_BONE_INDEX)
 				{
 					FTransform ParentLocalTransform = Keyframe->Get()->Pose.LocalTransforms[SourceParentIdx];
 					SourceComponentTransform = SourceComponentTransform * ParentLocalTransform;
@@ -65,8 +65,8 @@ void FAnimNextCopyBonesComponentSpaceTask::Execute(UE::UAF::FEvaluationVM& VM) c
 				// Compute the Destination Parent Component Space Transform
 
 				FTransform DestinationComponentParentTransform = FTransform::Identity;
-				int32 DestinationParentIdx = Keyframe->Get()->Pose.GetRefPose().GetLODParentBoneIndex(Keyframe->Get()->Pose.LODLevel, DestinationBoneIdx);
-				while (DestinationParentIdx != INDEX_NONE)
+				FBoneIndexType DestinationParentIdx = Keyframe->Get()->Pose.GetRefPose().GetLODParentBoneIndex(Keyframe->Get()->Pose.LODLevel, DestinationBoneIdx);
+				while (DestinationParentIdx != INVALID_BONE_INDEX)
 				{
 					FTransform ParentLocalTransform = Keyframe->Get()->Pose.LocalTransforms[DestinationParentIdx];
 					DestinationComponentParentTransform = DestinationComponentParentTransform * ParentLocalTransform;
diff --git a/Engine/Plugins/Experimental/UAF/UAFControlRig/Source/UAFControlRig/Private/AnimNextControlRigPoseAdapter.cpp b/Engine/Plugins/Experimental/UAF/UAFControlRig/Source/UAFControlRig/Private/AnimNextControlRigPoseAdapter.cpp
--- a/Engine/Plugins/Experimental/UAF/UAFControlRig/Source/UAFControlRig/Private/AnimNextControlRigPoseAdapter.cpp
+++ b/Engine/Plugins/Experimental/UAF/UAFControlRig/Source/UAFControlRig/Private/AnimNextControlRigPoseAdapter.cpp
@@ -134,7 +134,7 @@ void FAnimNextControlRigPoseAdapter::UpdateInputOutputMappingIfRequired(UControl
 			{
 				const int32 HierarchyParentIndex = InHierarchy->GetFirstParent(BoneElementIndex);
 				const int16 PoseIndex = ElementIndexToPoseIndex.FindChecked(BoneElementIndex);
-				const FCompactPoseBoneIndex CompactPoseParentIndex(ParentPoseIndices[PoseIndex]);
+				const FBoneIndexType ParentIndex(ParentPoseIndices[PoseIndex]);
 
 				FName HierarchyParentName(NAME_None);
 				FName PoseParentName(NAME_None);
@@ -143,12 +143,12 @@ void FAnimNextControlRigPoseAdapter::UpdateInputOutputMappingIfRequired(UControl
 				{
 					HierarchyParentName = InHierarchy->Get(HierarchyParentIndex)->GetFName();
 				}
-				if (CompactPoseParentIndex.IsValid())
+				if (ParentIndex != INVALID_BONE_INDEX)
 				{
-					const FSkeletonPoseBoneIndex SkeletonIndex(InRefPose.GetSkeletonBoneIndexFromLODBoneIndex(CompactPoseParentIndex.GetInt()));
-					if (SkeletonIndex.IsValid() && RefSkeleton->IsValidIndex(SkeletonIndex.GetInt()))
+					FBoneIndexType SkeletonIndex = InRefPose.GetSkeletonBoneIndexFromLODBoneIndex(ParentIndex);
+					if (SkeletonIndex != INVALID_BONE_INDEX && RefSkeleton->IsValidIndex((int32)SkeletonIndex))
 					{
-						PoseParentName = RefSkeleton->GetBoneName(SkeletonIndex.GetInt());
+						PoseParentName = RefSkeleton->GetBoneName((int32)SkeletonIndex);
 					}
 				}
 
diff --git a/Engine/Plugins/Experimental/UAF/UAFMirroring/Source/UAFMirroring/Private/Mirroring.cpp b/Engine/Plugins/Experimental/UAF/UAFMirroring/Source/UAFMirroring/Private/Mirroring.cpp
--- a/Engine/Plugins/Experimental/UAF/UAFMirroring/Source/UAFMirroring/Private/Mirroring.cpp
+++ b/Engine/Plugins/Experimental/UAF/UAFMirroring/Source/UAFMirroring/Private/Mirroring.cpp
@@ -63,7 +63,7 @@ namespace UE::UAF
 			const FBoneIndexType TargetBoneIndex = InMeshBoneIndexToMirroredMeshBoneIndexMap[SourceBoneIndex];
 
 			// Not mapped, so skip.
-			if (TargetBoneIndex == FBoneIndexType(INDEX_NONE))
+			if (TargetBoneIndex == INVALID_BONE_INDEX)
 			{
 				continue;
 			}
@@ -123,11 +123,11 @@ namespace UE::UAF
 		// Build mirror map for all the mesh bones in reference pose (LOD0).
 		for (FBoneIndexType MeshBoneIndex = 0; MeshBoneIndex < OutMeshBoneIndexToMirroredMeshBoneIndexMap.Num(); ++MeshBoneIndex)
 		{
-			if (OutMeshBoneIndexToMirroredMeshBoneIndexMap[MeshBoneIndex] == FBoneIndexType(INDEX_NONE))
+			if (OutMeshBoneIndexToMirroredMeshBoneIndexMap[MeshBoneIndex] == INVALID_BONE_INDEX)
 			{
 				// Get index for LOD sorted map
 				const FBoneIndexType LODBoneIndex = MeshBoneIndexToLODBoneIndexMap[MeshBoneIndex];
-				checkf(LODBoneIndex != FBoneIndexType(INDEX_NONE), TEXT("Failed to find LOD Bone Index counterpart for Mesh Bone Index"))
+				checkf(LODBoneIndex != INVALID_BONE_INDEX, TEXT("Failed to find LOD Bone Index counterpart for Mesh Bone Index"))
 
 				// We need to convert from our sorted LOD bone map to the skeleton bone map, otherwise the mapping is broken.
 				const FBoneIndexType SkeletonBoneIndex = LODBoneIndexToSkeletonBoneIndexMap[LODBoneIndex];
@@ -135,14 +135,14 @@ namespace UE::UAF
 				// since this is a skeleton index we need to use the skeleton to find the name 
 				// Find the candidate mirror partner for this bone
 				const FName SourceBoneName = FullLODSkeleton.GetBoneName(SkeletonBoneIndex);
-				FBoneIndexType MirrorMeshBoneIndex = INDEX_NONE;
+				FBoneIndexType MirrorMeshBoneIndex = INVALID_BONE_INDEX;
 				
 				FName* MirroredBoneName = NameToMirrorNameBoneMap.Find(SourceBoneName);
 				if (!SourceBoneName.IsNone() && MirroredBoneName && BoneNameToLODBoneIndexMap.Contains(*MirroredBoneName))
 				{
 					int32 MirrorLODBoneIndex = BoneNameToLODBoneIndexMap[*MirroredBoneName];
 
-					if (MirrorLODBoneIndex != FBoneIndexType(INDEX_NONE))
+					if (MirrorLODBoneIndex != INVALID_BONE_INDEX)
 					{
 						MirrorMeshBoneIndex = LODBoneIndexToMeshBoneIndexMap[MirrorLODBoneIndex];
 					}
@@ -152,7 +152,7 @@ namespace UE::UAF
 				
 				// Map candidate mirror partner to current bone.
 				// @todo: What happens we have conflicts on mirror data table setup. Should we have a conflict policy to resolve this? For example, two bones map to another.
-				if (MirrorMeshBoneIndex != FBoneIndexType(INDEX_NONE))
+				if (MirrorMeshBoneIndex != INVALID_BONE_INDEX)
 				{
 					OutMeshBoneIndexToMirroredMeshBoneIndexMap[MirrorMeshBoneIndex] = MeshBoneIndex;
 				}
@@ -266,8 +266,8 @@ namespace UE::UAF
 		// Contract/Assumption: the (SourceBoneIndex -> TargetBoneIndex) mapping used to call this matches the mapping used to build OutMirrorCorrections (i.e., Target == MirrorMap[Source]).
 		const auto MirrorTransform = [&InRefPoseMeshSpaceRotations, &InMirrorAxis, InRefPoseMeshSpaceRotationCorrections](const FTransform& SourceTransform, const FBoneIndexType& SourceParentIndex, const FBoneIndexType& SourceBoneIndex, const FBoneIndexType& TargetParentIndex, const FBoneIndexType& TargetBoneIndex) -> FTransform
 		{
-			const FQuat TargetParentRefRotation = TargetParentIndex != FBoneIndexType(INDEX_NONE) ? InRefPoseMeshSpaceRotations[TargetParentIndex] : FQuat::Identity;
-			const FQuat SourceParentRefRotation = SourceParentIndex != FBoneIndexType(INDEX_NONE) ? InRefPoseMeshSpaceRotations[SourceParentIndex] : FQuat::Identity;
+			const FQuat TargetParentRefRotation = TargetParentIndex != INVALID_BONE_INDEX ? InRefPoseMeshSpaceRotations[TargetParentIndex] : FQuat::Identity;
+			const FQuat SourceParentRefRotation = SourceParentIndex != INVALID_BONE_INDEX ? InRefPoseMeshSpaceRotations[SourceParentIndex] : FQuat::Identity;
 
 			// Mirror the translation component: Rotate the translation into the space of the mirror plane, mirror across the plane, and rotate into the space of its new parent.
 
@@ -303,11 +303,11 @@ namespace UE::UAF
 			const FBoneIndexType RootLODBoneIndex =  InMeshBoneIndexToMirroredMeshBoneIndexMap[RootMeshBoneIndex];
 			const FBoneIndexType MirrorRootMeshBoneIndex = InMeshBoneIndexToMirroredMeshBoneIndexMap[RootMeshBoneIndex];
 
-			if (MirrorRootMeshBoneIndex != FBoneIndexType(INDEX_NONE) || RootLODBoneIndex != FBoneIndexType(INDEX_NONE))
+			if (MirrorRootMeshBoneIndex != INVALID_BONE_INDEX && RootLODBoneIndex != INVALID_BONE_INDEX)
 			{
 				const FBoneIndexType MirroredRootLODBoneIndex = MeshBoneIndexToLODBoneIndexMap[MirrorRootMeshBoneIndex];
 
-				if (MirroredRootLODBoneIndex == FBoneIndexType(INDEX_NONE))
+				if (MirroredRootLODBoneIndex == INVALID_BONE_INDEX)
 				{
 					UE_LOGFMT(LogAnimation, Error, "UAF::MirrorPose - Failed to find LOD Bone index ({0}) for Mesh Bone Index ({1}) for a mirrored root bone. This is invalid.", MirroredRootLODBoneIndex, MirrorRootMeshBoneIndex);
 					return;
@@ -316,13 +316,13 @@ namespace UE::UAF
 				const FBoneIndexType TargetParentLODBoneIndex = LODBoneIndexToParentLODBoneIndexMap[RootMeshBoneIndex];
 				const FBoneIndexType SourceParentLODBoneIndex = LODBoneIndexToParentLODBoneIndexMap[MirroredRootLODBoneIndex];
 
-				if (TargetParentLODBoneIndex != FBoneIndexType(INDEX_NONE))
+				if (TargetParentLODBoneIndex != INVALID_BONE_INDEX)
 				{
 					UE_LOGFMT(LogAnimation, Error, "UAF::MirrorPose - Found parent bone index ({0}) for root bone index ({1}). This is invalid.", TargetParentLODBoneIndex, RootMeshBoneIndex);
 					return;
 				}
 
-				if (SourceParentLODBoneIndex != FBoneIndexType(INDEX_NONE))
+				if (SourceParentLODBoneIndex != INVALID_BONE_INDEX)
 				{
 					UE_LOGFMT(LogAnimation, Error, "UAF::MirrorPose - Found parent bone index ({0}) for mirrored root bone index ({1}). This is invalid.", TargetParentLODBoneIndex, RootMeshBoneIndex);
 					return;
@@ -340,8 +340,8 @@ namespace UE::UAF
 		for (FBoneIndexType LODBoneIndex = 1; LODBoneIndex < SourceLODBoneNum; ++LODBoneIndex)
 		{
 			// Skip invalid mesh bone index (should never happen).
-			const FBoneIndexType MeshBoneIndex = LODBoneIndex != FBoneIndexType(INDEX_NONE) ? LODBoneIndexToMeshBoneIndexMap[LODBoneIndex] : FBoneIndexType(INDEX_NONE);
-			if (MeshBoneIndex == FBoneIndexType(INDEX_NONE))
+			const FBoneIndexType MeshBoneIndex = LODBoneIndex != INVALID_BONE_INDEX ? LODBoneIndexToMeshBoneIndexMap[LODBoneIndex] : INVALID_BONE_INDEX;
+			if (MeshBoneIndex == INVALID_BONE_INDEX)
 			{
 				continue;
 			}
@@ -350,7 +350,7 @@ namespace UE::UAF
 			const FBoneIndexType MirroredMeshBoneIndex = InMeshBoneIndexToMirroredMeshBoneIndexMap[MeshBoneIndex];
 
 			// Skip invalid mirrors.
-			const bool bIsMirrorInvalid = MirroredMeshBoneIndex == FBoneIndexType(INDEX_NONE);
+			const bool bIsMirrorInvalid = MirroredMeshBoneIndex == INVALID_BONE_INDEX;
 			if (bIsMirrorInvalid)
 			{
 				continue;	
@@ -358,7 +358,7 @@ namespace UE::UAF
 			
 			const FBoneIndexType MirroredLODBoneIndex = MeshBoneIndexToLODBoneIndexMap[MirroredMeshBoneIndex];
 			const FBoneIndexType ParentLODBoneIndex = LODBoneIndexToParentLODBoneIndexMap[LODBoneIndex];
-			const FBoneIndexType ParentMeshBoneIndex = ParentLODBoneIndex != FBoneIndexType(INDEX_NONE) ? LODBoneIndexToMeshBoneIndexMap[ParentLODBoneIndex] : FBoneIndexType(INDEX_NONE);
+			const FBoneIndexType ParentMeshBoneIndex = ParentLODBoneIndex != INVALID_BONE_INDEX ? LODBoneIndexToMeshBoneIndexMap[ParentLODBoneIndex] : INVALID_BONE_INDEX;
 			
 			// Self-mirror (mirror in place relative to the same parent).
 			const bool bIsMirroredBoneMappedToSelf = LODBoneIndex == MirroredLODBoneIndex;
diff --git a/Engine/Plugins/Experimental/UAF/UAFPoseSearch/Source/UAFPoseSearch/Private/HistoryCollectorTrait.cpp b/Engine/Plugins/Experimental/UAF/UAFPoseSearch/Source/UAFPoseSearch/Private/HistoryCollectorTrait.cpp
--- a/Engine/Plugins/Experimental/UAF/UAFPoseSearch/Source/UAFPoseSearch/Private/HistoryCollectorTrait.cpp
+++ b/Engine/Plugins/Experimental/UAF/UAFPoseSearch/Source/UAFPoseSearch/Private/HistoryCollectorTrait.cpp
@@ -63,7 +63,7 @@ namespace UE::UAF
 			const TArrayView<const FBoneIndexType> SkeletonBoneIndexToLODBoneIndexMap = Pose.GetSkeletonBoneIndexToLODBoneIndexMap();
 			const FBoneIndexType LODBoneIndex = SkeletonBoneIndexToLODBoneIndexMap[SkeletonBoneIdx.GetInt()];
 			
-			if (LODBoneIndex != FBoneIndexType(INDEX_NONE) && LODBoneIndex < Pose.GetNumBones())
+			if (LODBoneIndex != INVALID_BONE_INDEX && LODBoneIndex < Pose.GetNumBones())
 			{
 				ComponentSpaceTransform = Pose.LocalTransformsView[LODBoneIndex];
 			}
diff --git a/Engine/Plugins/Experimental/UAF/UAFWarping/Source/UAFWarping/Private/StrafeWarpingTrait.cpp b/Engine/Plugins/Experimental/UAF/UAFWarping/Source/UAFWarping/Private/StrafeWarpingTrait.cpp
--- a/Engine/Plugins/Experimental/UAF/UAFWarping/Source/UAFWarping/Private/StrafeWarpingTrait.cpp
+++ b/Engine/Plugins/Experimental/UAF/UAFWarping/Source/UAFWarping/Private/StrafeWarpingTrait.cpp
@@ -115,7 +115,7 @@ namespace UE::UAF
 			}
 			else
 			{
-				OutSpineBoneData[i].LODBoneIndex = INDEX_NONE;
+				OutSpineBoneData[i].LODBoneIndex = INVALID_BONE_INDEX;
 			}
 		}
 
@@ -462,7 +462,7 @@ void FAnimNextStrafeWarpingTask::DoWarpPose(TUniquePtr<UE::UAF::FKeyframeState>&
 				{
 					const FStrafeWarpFootData& FootData = SharedData->FootData[ArrayIndex];
 					FBoneIndexType LegTipIndex = Pose.FindLODBoneIndexFromBoneName(FootData.LegTip);
-					if (LegTipIndex == INDEX_NONE)
+					if (LegTipIndex == INVALID_BONE_INDEX)
 					{
 						continue;
 					}
@@ -484,7 +484,7 @@ void FAnimNextStrafeWarpingTask::DoWarpPose(TUniquePtr<UE::UAF::FKeyframeState>&
 						FBoneIndexType LegMidIndex = Pose.FindLODBoneIndexFromBoneName(FootData.LegMid);
 
 						// Validate data
-						if (LegRootIndex == INDEX_NONE || LegMidIndex == INDEX_NONE)
+						if (LegRootIndex == INVALID_BONE_INDEX || LegMidIndex == INVALID_BONE_INDEX)
 						{
 							continue;
 						}
@@ -561,11 +561,11 @@ void FAnimNextStrafeWarpingTask::DoWarpPose(TUniquePtr<UE::UAF::FKeyframeState>&
 
 					FBoneIndexType IKTargetIndex = Pose.FindLODBoneIndexFromBoneName(FootData.IKTargetBone);
 
-					if (IKTargetIndex != INDEX_NONE)
+					if (IKTargetIndex != INVALID_BONE_INDEX)
 					{
 						// Set IK target to equal leg tip in component space
 						FBoneIndexType IKTargetParent = Pose.GetLODBoneParentIndex(IKTargetIndex);		// Usually the root bone but we can make it work for any parent
-						if (IKTargetParent != INDEX_NONE)
+						if (IKTargetParent != INVALID_BONE_INDEX)
 						{
 							FTransform IKTargetCS = FootTargetTransform;
 							FTransform IKTargetParentCS = ComponentSpaceTransforms[IKTargetParent];
@@ -624,7 +624,7 @@ void FAnimNextStrafeWarpingTask::DoWarpPose(TUniquePtr<UE::UAF::FKeyframeState>&
 		for (int32 ArrayIndex = NumSpineBones - 1; ArrayIndex >= 0; ArrayIndex--)
 		{
 			const FStrafeWarpingTrait::FSpineBoneData& BoneData = SpineBoneDataArray[ArrayIndex];
-			if (BoneData.LODBoneIndex == INDEX_NONE || BoneData.Weight == 0.0f)
+			if (BoneData.LODBoneIndex == INVALID_BONE_INDEX || BoneData.Weight == 0.0f)
 			{
 				continue;
 			}
diff --git a/Engine/Source/Runtime/AnimationCore/Public/BoneIndices.h b/Engine/Source/Runtime/AnimationCore/Public/BoneIndices.h
--- a/Engine/Source/Runtime/AnimationCore/Public/BoneIndices.h
+++ b/Engine/Source/Runtime/AnimationCore/Public/BoneIndices.h
@@ -6,6 +6,9 @@
 
 typedef uint16 FBoneIndexType;
 
+// Important: only use this define when comparing FBoneIndexType (16 bit), it should not be used to compare 32 bit bone indicies stored in FBoneIndexBase
+#define INVALID_BONE_INDEX FBoneIndexType(INDEX_NONE) 
+
 
 struct FBoneIndexBase
 {
